Updates for Modular Class Display to add a widget registration system and get it working for BaseToolkit and toolbars.

diff --git a/Engine/Plugins/Editor/ModelingToolsEditorMode/Source/ModelingToolsEditorMode/ModelingToolsEditorMode.Build.cs b/Engine/Plugins/Editor/ModelingToolsEditorMode/Source/ModelingToolsEditorMode/ModelingToolsEditorMode.Build.cs
--- a/Engine/Plugins/Editor/ModelingToolsEditorMode/Source/ModelingToolsEditorMode/ModelingToolsEditorMode.Build.cs
+++ b/Engine/Plugins/Editor/ModelingToolsEditorMode/Source/ModelingToolsEditorMode/ModelingToolsEditorMode.Build.cs
@@ -46,6 +46,7 @@ public class ModelingToolsEditorMode : ModuleRules
 				"ToolWidgets",
 				"EditorWidgets",
 				"ModelingEditorUI",
+				"WidgetRegistration",
 				"StylusInput",
 				"DeveloperSettings",
 				"PropertyEditor",
diff --git a/Engine/Source/Developer/WidgetRegistration/Private/ToolElementRegistry.cpp b/Engine/Source/Developer/WidgetRegistration/Private/ToolElementRegistry.cpp
new file mode 100644
--- /dev/null
+++ b/Engine/Source/Developer/WidgetRegistration/Private/ToolElementRegistry.cpp
@@ -0,0 +1,116 @@
+// Copyright Epic Games, Inc. All Rights Reserved.
+
+#pragma once
+
+#include "ToolElementRegistry.h"
+
+#include "Widgets/SNullWidget.h"
+
+void FToolElementRegistrationArgs::UpdateWidget()
+{
+	
+}
+
+void FToolElementRegistrationArgs::ResetWidget()
+{
+	
+}
+
+FToolElementRegistrationArgs::FToolElementRegistrationArgs(
+	EToolElement InToolElementType) :
+	ToolElementType(InToolElementType)
+{
+}
+	
+TSharedPtr<SWidget> FToolElementRegistrationArgs::GenerateWidget()
+{
+	return nullptr;
+}
+
+FToolElementRegistrationKey::FToolElementRegistrationKey(FName InName, EToolElement InToolElementType ) :
+		Name(InName), ToolElementType(InToolElementType)
+{
+	FString KeyStr = Name.ToString().AppendChar(' ');
+	KeyStr.AppendInt( static_cast<int32>(ToolElementType));
+	KeyString = KeyStr;
+}
+
+FString FToolElementRegistrationKey::GetKeyString()
+{
+	return KeyString;	
+}
+
+FToolElement::FToolElement(
+		const FName InName,
+		TSharedRef<FToolElementRegistrationArgs> InBuilderArgs) :
+		FToolElementRegistrationKey(
+			InName,
+			InBuilderArgs->ToolElementType),
+		RegistrationArgs(InBuilderArgs)
+{
+}
+
+void FToolElement::SetRegistrationArgs(TSharedRef<FToolElementRegistrationArgs> InRegistrationArgs)
+{
+	RegistrationArgs = InRegistrationArgs;
+}
+
+TSharedRef<SWidget> FToolElement::GenerateWidget()
+{
+	return RegistrationArgs->GenerateWidget().ToSharedRef();
+}
+
+FToolElementRegistry::FToolElementRegistry()
+{
+	ToolElementKeyToToolElementMap.Reset();
+}
+
+FToolElementRegistry& FToolElementRegistry::Get()
+{
+	static FToolElementRegistry Registry;
+	return Registry;
+}
+
+TSharedPtr<FToolElement> FToolElementRegistry::GetToolElementSP(FToolElementRegistrationKey& ToolElementKey)
+{
+	TSharedPtr<FToolElement>* Element = ToolElementKeyToToolElementMap.Find(*(ToolElementKey.GetKeyString()));
+	return Element != nullptr ? *Element : nullptr;
+}
+
+TSharedRef< SWidget > FToolElementRegistry::GenerateWidget(
+		TSharedRef<FToolElementRegistrationKey> ToolElementKeySR,
+		TSharedPtr<FToolElementRegistrationArgs> RegistrationArgsSP,
+		bool bUpdateRegistrationArgs)
+{
+	if (const TSharedPtr<FToolElement> Element = GetToolElementSP(ToolElementKeySR.Get()))
+	{
+		const bool bHasNewRegistrationArgs = RegistrationArgsSP != nullptr;
+		const TSharedRef<FToolElementRegistrationArgs> TempArgsSR = Element->RegistrationArgs;
+			
+		if (bHasNewRegistrationArgs)
+		{
+			Element->RegistrationArgs = RegistrationArgsSP.ToSharedRef();		
+		}
+		const TSharedPtr<SWidget> Widget = Element->RegistrationArgs->GenerateWidget();
+
+		if ( !bUpdateRegistrationArgs && bHasNewRegistrationArgs )
+		{
+			Element->RegistrationArgs = TempArgsSR;
+		}
+		
+		return Widget != nullptr &&  Widget.IsValid() ? Widget.ToSharedRef() : SNullWidget::NullWidget;
+	}
+
+	return SNullWidget::NullWidget;
+}
+
+
+void FToolElementRegistry::RegisterElement(const TSharedRef<FToolElement> ToolElement)
+{
+	ToolElementKeyToToolElementMap.Add(*ToolElement->GetKeyString(), ToolElement);
+}
+
+void FToolElementRegistry::UnregisterElement(const TSharedRef<FToolElement> ToolElement)
+{
+	ToolElementKeyToToolElementMap.Remove(*ToolElement->GetKeyString());
+}
diff --git a/Engine/Source/Developer/WidgetRegistration/Private/ToolbarRegistrationArgs.cpp b/Engine/Source/Developer/WidgetRegistration/Private/ToolbarRegistrationArgs.cpp
new file mode 100644
--- /dev/null
+++ b/Engine/Source/Developer/WidgetRegistration/Private/ToolbarRegistrationArgs.cpp
@@ -0,0 +1,17 @@
+// Copyright Epic Games, Inc. All Rights Reserved.
+
+#pragma once
+#include "ToolbarRegistrationArgs.h"
+
+
+FToolbarRegistrationArgs::FToolbarRegistrationArgs(TSharedRef<FToolBarBuilder> InToolBarBuilder) :
+		FToolElementRegistrationArgs(EToolElement::Toolbar),
+		ToolBarBuilder(InToolBarBuilder)
+{
+		
+}
+
+TSharedPtr<SWidget> FToolbarRegistrationArgs::GenerateWidget()
+{
+	return ToolBarBuilder->MakeWidget();
+}
\ No newline at end of file
diff --git a/Engine/Source/Developer/WidgetRegistration/Private/ToolkitBuilder.cpp b/Engine/Source/Developer/WidgetRegistration/Private/ToolkitBuilder.cpp
new file mode 100644
--- /dev/null
+++ b/Engine/Source/Developer/WidgetRegistration/Private/ToolkitBuilder.cpp
@@ -0,0 +1,281 @@
+
+
+#include "ToolkitBuilder.h"
+
+#include "Widgets/SBoxPanel.h"
+#include "SPrimaryButton.h"
+#include "ToolbarRegistrationArgs.h"
+#include "ToolElementRegistry.h"
+#include "Framework/Commands/UICommandList.h"
+#include "Styling/StyleColors.h"
+
+#define LOCTEXT_NAMESPACE "ToolkitBuilder"
+
+FToolElementRegistry FToolkitBuilder::ToolRegistry = FToolElementRegistry::Get();
+
+bool FEditablePalette::IsInPalette(const FName CommandName) const
+{
+	return this->PaletteCommandNameArray.Contains(CommandName.ToString());
+}
+
+FEditablePalette::FEditablePalette(TSharedPtr<FUICommandInfo> InLoadToolPaletteAction,
+	TArray<FString>& InPaletteCommandNameArray,
+	TSharedPtr<FUICommandInfo> InAddToPaletteAction,
+	TSharedPtr<FUICommandInfo> InRemoveFromPaletteAction) :
+	FToolPalette(InLoadToolPaletteAction, {}),
+	AddToPaletteAction(InAddToPaletteAction),
+	RemoveFromPaletteAction(InRemoveFromPaletteAction),
+	PaletteCommandNameArray(InPaletteCommandNameArray)
+{
+}
+
+FToolkitBuilder::FToolkitBuilder(
+	FName InToolbarCustomizationName,
+	TSharedPtr<FUICommandList> InToolkitCommandList) :
+	FToolElementRegistrationArgs(EToolElement::Toolkit),
+	ToolbarCustomizationName(InToolbarCustomizationName),
+	ToolkitCommandList(InToolkitCommandList)
+{
+	ResetWidget();
+}
+
+TSharedPtr<FToolBarBuilder> FToolkitBuilder::GetLoadPaletteToolbar()
+{
+	return LoadPaletteToolBarBuilder;
+}
+
+TSharedRef<SWidget> FToolkitBuilder::CreateToolbarWidget() const
+{
+	return ToolRegistry.GenerateWidget(VerticalToolbarElement.ToSharedRef());
+}
+
+void FToolkitBuilder::AddPalette(TSharedPtr<FEditablePalette> Palette)
+{
+	EditablePalettesArray.Add(Palette.ToSharedRef());
+	AddPalette(StaticCastSharedRef<FToolPalette>(Palette.ToSharedRef()) );
+}
+
+
+void FToolkitBuilder::AddPalette(TSharedPtr<FToolPalette> Palette)
+{
+	for (TSharedRef<FButtonArgs> Button : Palette->PaletteActions)
+	{
+		PaletteCommandNameToButtonArgsMap.Add(Button->Command->GetCommandName().ToString(), Button);
+	}
+
+	LoadToolPaletteCommandList->MapAction(
+				Palette->LoadToolPaletteAction,
+				FExecuteAction::CreateSP(SharedThis(this),
+				&FToolkitBuilder::TogglePalette,
+				Palette),
+				FCanExecuteAction::CreateLambda([] { return true; }),
+				FGetActionCheckState::CreateSP(SharedThis(this),
+					&FToolkitBuilder::IsActiveToolPalette,
+					Palette->LoadToolPaletteAction->GetCommandName())
+									 	);
+	LoadPaletteToolBarBuilder->AddToolBarButton(Palette->LoadToolPaletteAction);
+}
+
+ECheckBoxState FToolkitBuilder::IsActiveToolPalette(FName CommandName) const
+{
+	return (ActivePalette &&
+			ActivePalette->LoadToolPaletteAction->GetCommandName() == CommandName) ?
+				ECheckBoxState::Checked : ECheckBoxState::Unchecked;
+}
+
+void FToolkitBuilder::UpdateEditablePalette(FEditablePalette& Palette)
+{
+	Palette.PaletteActions.Empty();
+
+	for (const FString& Key : Palette.PaletteCommandNameArray)
+	{
+		if (const TSharedPtr<FButtonArgs>* FoundButton = PaletteCommandNameToButtonArgsMap.Find(Key))
+		{
+			const TSharedRef<FButtonArgs> Button = (*FoundButton).ToSharedRef();
+			Palette.PaletteActions.Add(Button);
+		}
+	}
+}
+
+void FToolkitBuilder::UpdateWidget()
+{
+	for (const TSharedRef<FEditablePalette> EditablePalette : EditablePalettesArray)
+	{
+		UpdateEditablePalette(*EditablePalette);		
+	}
+	
+}
+
+void FToolkitBuilder::ToggleCommandInPalette(TSharedRef<FEditablePalette> Palette, FString CommandNameString)
+{
+	if (!Palette->PaletteCommandNameArray.Contains(CommandNameString))
+	{	
+		Palette->PaletteCommandNameArray.Add(CommandNameString);
+	}
+	else
+	{
+		Palette->PaletteCommandNameArray.Remove(CommandNameString);		
+	}
+
+	UpdateEditablePalette(*Palette);
+
+	// if the active Palette is the Palette to which we are toggling an action,
+	// recreate it to load the new state after the toggle
+	if ( ActivePalette == Palette )
+	{
+		CreatePalette(Palette);
+	}
+}
+
+bool FToolkitBuilder::HasActivePalette() const
+{
+	return ActivePalette != nullptr;
+}
+
+
+void FToolkitBuilder::TogglePalette(TSharedPtr<FToolPalette> Palette)
+{
+	const FName CommandName = Palette->LoadToolPaletteAction->GetCommandName();
+	
+	if (ActivePalette && ActivePalette->LoadToolPaletteAction->GetCommandName() == CommandName)
+	{
+		ActivePalette = nullptr;
+		ResetToolPaletteWidget();
+	    OnToolEnded.ExecuteIfBound();
+		return;
+	}
+	
+	CreatePalette(Palette);
+}
+
+void FToolkitBuilder::CreatePalette(TSharedPtr<FToolPalette> Palette)
+{
+	const FName CommandName = Palette->LoadToolPaletteAction->GetCommandName();
+	ActivePalette = Palette;
+	ResetToolPaletteWidget();
+	const TSharedPtr<FSlimHorizontalUniformToolBarBuilder> PaletteToolbarBuilder =
+		MakeShareable(new FSlimHorizontalUniformToolBarBuilder(ToolkitCommandList,
+			FMultiBoxCustomization(ToolbarCustomizationName)));
+
+	const TSharedPtr<FToolbarRegistrationArgs> RegistrationArgs = MakeShareable<FToolbarRegistrationArgs>(
+		new FToolbarRegistrationArgs(PaletteToolbarBuilder.ToSharedRef()));
+	FToolElementRegistrationKey Key = FToolElementRegistrationKey(CommandName, EToolElement::Toolbar);
+	TSharedPtr<FToolElement> Element = ToolRegistry.GetToolElementSP(Key);
+	
+	if (!Element)
+	{
+		Element = MakeShareable(
+				new FToolElement(Palette->LoadToolPaletteAction->GetCommandName(),
+				RegistrationArgs.ToSharedRef()));
+		ToolRegistry.RegisterElement(Element.ToSharedRef());
+	}
+	
+	Element->SetRegistrationArgs(RegistrationArgs.ToSharedRef());
+	
+	LoadCommandNameToPaletteToolbarBuilderMap.Add(Palette->LoadToolPaletteAction->GetCommandName(), PaletteToolbarBuilder.ToSharedRef());
+	
+	PaletteToolbarBuilder->SetStyle(&FAppStyle::Get(), "SlimPaletteToolBar");
+
+	for (TSharedRef<FButtonArgs> PaletteButton : Palette->PaletteActions)
+	{
+		PaletteButton->CommandList = ToolkitCommandList;
+		PaletteButton->UserInterfaceActionType = EUserInterfaceActionType::ToggleButton;
+		PaletteButton->OnGetMenuContent.BindSP(SharedThis(this),
+			&FToolkitBuilder::GetContextMenuContent,
+			PaletteButton->Command->GetCommandName());
+		PaletteToolbarBuilder->AddToolBarButton(PaletteButton.Get());
+	}
+
+	ToolPaletteWidget->AddSlot()
+		.HAlign(HAlign_Fill)
+		.Padding(0.f)
+		.FillHeight(1.0f)
+		[
+
+		SNew(SBorder)
+		.Padding(FAppStyle::GetMargin("SlimPaletteToolBarStyle.TitleBar.Padding"))
+		.VAlign(VAlign_Center)
+		.BorderImage(FAppStyle::GetBrush("ToolPanel.GroupBorder"))
+		.HAlign(HAlign_Left)
+		[
+
+			SNew(STextBlock)
+					.Justification(ETextJustify::Left)
+					.AutoWrapText(true)
+					.Text(Palette->LoadToolPaletteAction->GetLabel())
+					.ColorAndOpacity(FStyleColors::Foreground)
+		]];
+	
+	ToolPaletteWidget->AddSlot()
+		.HAlign(HAlign_Fill)
+		.AutoHeight()
+		[
+			Element->GenerateWidget()
+		];
+
+	OnToolEnded.ExecuteIfBound();
+}
+
+TSharedRef<SWidget> FToolkitBuilder::GetToolPaletteWidget() const
+{
+	return ToolPaletteWidget->AsShared();
+}
+
+TSharedRef<SWidget> FToolkitBuilder::GetContextMenuContent(const FName CommandName)
+{
+	FMenuBuilder MenuBuilder(true, nullptr);
+
+	for (const TSharedRef<FEditablePalette> EditablePalette : EditablePalettesArray)
+	{
+		const FUIAction ItemAction(FExecuteAction::CreateSP(
+			SharedThis(this), 
+			&FToolkitBuilder::ToggleCommandInPalette, 
+			EditablePalette, CommandName.ToString()));
+		const FText LoadPaletteActionCommandLabelFText = EditablePalette->LoadToolPaletteAction->GetLabel();
+		const FText& ItemText = EditablePalette->IsInPalette(CommandName) ? 
+			FText::Format(LOCTEXT("RemoveFromPalette", "Remove from {0}"), LoadPaletteActionCommandLabelFText) : 
+			FText::Format(LOCTEXT("AddToPalette", "Add to {0}"), LoadPaletteActionCommandLabelFText);
+
+		MenuBuilder.AddMenuEntry(ItemText, ItemText, FSlateIcon(), ItemAction);
+		
+	}
+
+	return MenuBuilder.MakeWidget();
+}
+
+void FToolkitBuilder::ResetWidget()
+{
+	LoadToolPaletteCommandList = MakeShareable(new FUICommandList);
+	LoadPaletteToolBarBuilder = MakeShared<FVerticalToolBarBuilder>(LoadToolPaletteCommandList, FMultiBoxCustomization::None, TSharedPtr<FExtender>(), true);
+	ToolPaletteWidget = SNew(SVerticalBox);
+
+	LoadCommandNameToPaletteToolbarBuilderMap.Reset();
+	EditablePalettesArray.Reset();
+
+	FToolElementRegistrationKey Key = FToolElementRegistrationKey(ToolbarCustomizationName, EToolElement::Toolbar);
+	VerticalToolbarElement = ToolRegistry.GetToolElementSP(Key);
+	const TSharedRef<FToolbarRegistrationArgs> VerticalToolbarRegistrationArgs = MakeShareable<FToolbarRegistrationArgs>(
+		new FToolbarRegistrationArgs(LoadPaletteToolBarBuilder.ToSharedRef()));
+	
+	if (!VerticalToolbarElement)
+	{
+		VerticalToolbarElement = MakeShareable(new FToolElement
+			(ToolbarCustomizationName,
+			VerticalToolbarRegistrationArgs));
+		ToolRegistry.RegisterElement(VerticalToolbarElement.ToSharedRef());
+	}
+
+	VerticalToolbarElement->SetRegistrationArgs(VerticalToolbarRegistrationArgs);	
+}
+
+void FToolkitBuilder::ResetToolPaletteWidget()
+{
+	if (ToolPaletteWidget)
+	{
+		ToolPaletteWidget->ClearChildren();
+		return;
+	}
+	ToolPaletteWidget = SNew(SVerticalBox);
+}
+
+#undef LOCTEXT_NAMESPACE
\ No newline at end of file
diff --git a/Engine/Source/Developer/WidgetRegistration/Private/WidgetRegistrationModule.cpp b/Engine/Source/Developer/WidgetRegistration/Private/WidgetRegistrationModule.cpp
new file mode 100644
--- /dev/null
+++ b/Engine/Source/Developer/WidgetRegistration/Private/WidgetRegistrationModule.cpp
@@ -0,0 +1,29 @@
+// Copyright Epic Games, Inc. All Rights Reserved.
+
+// #include "ToolkitStyle.h"
+#include "Modules/ModuleInterface.h"
+#include "Modules/ModuleManager.h"
+
+// @TODO: ~enable ToolkitStyle set 
+
+class FWidgetRegistrationModule
+	: public IModuleInterface
+{
+public:
+
+	// IModuleInterface interface
+	
+	virtual void StartupModule() override
+	{
+	//	FToolkitStyle::Initialize();
+	}
+	
+	virtual void ShutdownModule() override
+	{
+	//	FToolkitStyle::Shutdown();
+	}
+
+};
+
+
+IMPLEMENT_MODULE(FWidgetRegistrationModule, WidgetRegistration);
diff --git a/Engine/Source/Developer/WidgetRegistration/Public/ToolElementRegistry.h b/Engine/Source/Developer/WidgetRegistration/Public/ToolElementRegistry.h
new file mode 100644
--- /dev/null
+++ b/Engine/Source/Developer/WidgetRegistration/Public/ToolElementRegistry.h
@@ -0,0 +1,159 @@
+//  Copyright Epic Games, Inc. All Rights Reserved.
+
+#pragma once
+
+#include "Widgets/SWidget.h"
+
+class FToolElementRegistry;
+
+/** An enum to provide the type of Tool Element*/
+enum class EToolElement : uint32
+{
+	Toolbar = 0,
+	Toolkit
+};
+
+/** A class which provides the data necessary to generate a widget. This class should be
+ * extended and used to hold whatever information is needed to generate or update the
+ * widgets that make up the tool, and implement the virtual widget methods
+ */
+class WIDGETREGISTRATION_API FToolElementRegistrationArgs :
+	public TSharedFromThis<FToolElementRegistrationArgs>
+{
+public:
+
+	/** The constructor, which takes an EToolElement which will define the type of Tool Element*/
+	FToolElementRegistrationArgs(EToolElement InToolElementType);
+
+	/** default destructor in case any subclasses need to provide a destructor*/
+	virtual ~FToolElementRegistrationArgs() = default;
+
+	/** Implements the generation of the TSharedPtr<SWidget> */
+	virtual TSharedPtr<SWidget> GenerateWidget();
+
+	/** Updates/reloads this widget. This should be called after a consumer has changed any Data in this */
+	virtual void UpdateWidget();
+
+	/** Resets the widget to its initial state */
+	virtual void ResetWidget();
+
+	/** The type of tool element this is Registration args for */
+	const EToolElement ToolElementType;
+};
+
+/** Serves as a key into the FToolElementRegistry of FToolElements */
+class WIDGETREGISTRATION_API FToolElementRegistrationKey :
+	public TSharedFromThis<FToolElementRegistrationKey> {
+
+public:
+
+	/** Constructor for FToolElementRegistrationKey.
+	 *
+	 * @param InName the FName that will be part of the compound key for the FToolElement this denotes
+	 * @param InToolElementType the EToolElementType that will be part of the compound key for the FToolElement this denotes
+	 */
+	FToolElementRegistrationKey(FName InName, EToolElement InToolElementType );
+
+	/** the FName that will be part of the compound key for the FToolElement this denotes */
+	FName Name;
+
+	/** the EToolElementType that will be part of the compound key for the FToolElement this denotes */ 
+	EToolElement ToolElementType;
+
+	FString GetKeyString();
+
+private:
+FString KeyString;
+};
+
+/** Represents one Tool Element, and is responsible for displaying the UI for that element  */
+class WIDGETREGISTRATION_API FToolElement : public FToolElementRegistrationKey, public TSharedFromThis<FToolElement>
+{
+
+public:
+
+	friend class FToolElementRegistry;
+
+	/** 
+	 * FToolElement constructor
+	 *
+	 * @param InName the FName of the FToolElement
+	 * @param InBuilderArgs the shared reference to the FToolElementRegistrationArgs. Need to have a shared reference
+	 * here so that if the registered outlives the registrar, we maintain a copy
+	 */
+	FToolElement(
+		const FName InName,
+		TSharedRef<FToolElementRegistrationArgs> InBuilderArgs);
+
+	/**
+	 * a convenience pass through method to the FToolElementRegistrationArgs GenerateWidget method
+	 * which generates the TSharedRef<SWidget>
+	 *
+	 * @return the TSharedRef<SWidget> which makes up the UI for this tool element
+	 */
+	TSharedRef<SWidget> GenerateWidget();
+
+	/*
+	 * Sets the FToolElementRegistrationArgs for the this tool element to RegistrationArgs
+	 *
+	 * @param RegistrationArgs the FToolElementRegistrationArgs for the this tool element
+	 */
+	void SetRegistrationArgs(TSharedRef<FToolElementRegistrationArgs> RegistrationArgs);
+
+private:
+	/** the FToolElementRegistrationArgs for the this tool element */
+	TSharedRef<FToolElementRegistrationArgs> RegistrationArgs;
+	
+};
+
+/** A Registry for FToolElements, where the FToolElements will generate the UI for a certain tool */
+class WIDGETREGISTRATION_API FToolElementRegistry 
+{
+public:
+
+	/**
+	 *Gets the TSharedPtr<FToolElement> for the given FToolElementRegistrationKey& tool element key
+	 *
+	 *@param ToolElementKey the given FToolElementRegistrationKey& tool element key for which to find the tool element
+	 */
+	TSharedPtr<FToolElement> GetToolElementSP( FToolElementRegistrationKey& ToolElementKey);
+	explicit FToolElementRegistry();
+
+	static FToolElementRegistry& Get();
+
+	/**
+	 * Generates the widget for the tool element denoted by the FToolElementRegistrationKey, using the specified
+	 * TSharedPtr<FToolElementRegistrationArgs> if provided (otherwise it will use the default registered registration
+	 * arguments. If bUpdateRegistrationArgs is set to true and RegistrationArgsSP is not nullptr, the specified
+	 * registration args will replace the existing ones for the tool element.
+	 *
+	 * @param ToolElementKeySR the TSharedRef<FToolElementRegistrationKey> which is the key to the tool element for 
+	        which we want to generate the widget
+	   @param RegistrationArgsSP The SharedPtr to the FToolElementRegistrationArgs for the tool element for which
+			which we want to generate the widget
+	 * @param bUpdateRegistrationArgs a bool that indicates whether we store any provided FToolElementRegistrationArgs
+	 * as the new FToolElementRegistrationArgs for this tool element
+	 */
+	TSharedRef< SWidget > GenerateWidget(
+		TSharedRef<FToolElementRegistrationKey> ToolElementKeySR,
+		TSharedPtr<FToolElementRegistrationArgs> RegistrationArgsSP = nullptr,
+		bool bUpdateRegistrationArgs = false);
+
+	/** registers the FToolElement ToolElement
+	 *
+	 * @param ToolElement the FToolElement that is to be registered
+	 */
+	void RegisterElement(const TSharedRef<FToolElement> ToolElement);
+	
+	/** Unregisters the FToolElement ToolElement
+	 *
+	 * @param ToolElement the FToolElement that is to be registered
+	 */
+	void UnregisterElement(const TSharedRef<FToolElement> ToolElement);
+
+private:
+	/* the Map of tool element Key to tool element */
+	TMap<FString, TSharedPtr<FToolElement>> ToolElementKeyToToolElementMap;
+};
+
+ 
\ No newline at end of file
diff --git a/Engine/Source/Developer/WidgetRegistration/Public/ToolbarRegistrationArgs.h b/Engine/Source/Developer/WidgetRegistration/Public/ToolbarRegistrationArgs.h
new file mode 100644
--- /dev/null
+++ b/Engine/Source/Developer/WidgetRegistration/Public/ToolbarRegistrationArgs.h
@@ -0,0 +1,31 @@
+//  Copyright Epic Games, Inc. All Rights Reserved.
+
+#pragma once
+
+#include "ToolElementRegistry.h"
+#include "Framework/MultiBox/MultiBoxBuilder.h"
+
+/** A class that provides the FToolbarRegistrationArgs for a UE toolbar  */
+class WIDGETREGISTRATION_API FToolbarRegistrationArgs : public FToolElementRegistrationArgs
+{
+public:
+
+	/** 
+	 * the constructor for FToolbarRegistrationArgs
+	 *
+	 * @param InToolBarBuilder a TSharedRef to the FToolBarBuilder  which will be used by
+	 * this to build up a toolbar SWidget
+	 */
+	FToolbarRegistrationArgs(TSharedRef<FToolBarBuilder> InToolBarBuilder);
+
+	/**
+	 * Generates the TSharedPtr<SWidget> for a toolbar
+	 *
+	 * @return the the TSharedPtr<SWidget> for a toolbar
+	 */
+	virtual TSharedPtr<SWidget> GenerateWidget() override;
+
+	/** The FToolBarBuilder object used to build up a toolbar widget */
+	TSharedRef<FToolBarBuilder> ToolBarBuilder;
+	
+};
diff --git a/Engine/Source/Developer/WidgetRegistration/Public/ToolkitBuilder.h b/Engine/Source/Developer/WidgetRegistration/Public/ToolkitBuilder.h
new file mode 100644
--- /dev/null
+++ b/Engine/Source/Developer/WidgetRegistration/Public/ToolkitBuilder.h
@@ -0,0 +1,191 @@
+// Copyright Epic Games, Inc. All Rights Reserved.
+
+#pragma once
+
+#include "ToolMenus.h"
+#include "Framework/Commands/UICommandInfo.h"
+#include "Widgets/SBoxPanel.h"
+#include "InteractiveToolManager.h"
+#include "ToolElementRegistry.h"
+
+
+class SWidget;
+class FToolElement;
+
+/** A struct that provides the data for a single tool Palette*/
+struct WIDGETREGISTRATION_API FToolPalette : TSharedFromThis<FToolPalette>
+{
+	FToolPalette(TSharedPtr<FUICommandInfo> InLoadToolPaletteAction,
+		const TArray<TSharedPtr< FUICommandInfo >>& InPaletteActions) :
+		LoadToolPaletteAction(InLoadToolPaletteAction)
+	{
+		for (const TSharedPtr< const FUICommandInfo > CommandInfo : InPaletteActions)
+		{
+			TSharedPtr<FButtonArgs> Button = MakeShareable(new FButtonArgs);
+			Button->Command = CommandInfo;
+			PaletteActions.Add(Button.ToSharedRef());
+		}
+	}
+
+	/** The FUICommandInfo button which loads a particular set of toads */
+	const TSharedPtr<FUICommandInfo> LoadToolPaletteAction;
+
+	/** The ButtonArgs that has the data to initialize the buttons in the FToolPalette loaded by LoadToolPaletteAction */
+	TArray<TSharedRef<FButtonArgs>> PaletteActions;
+
+	/** The FUICommandList associated with this Palette */
+	TSharedPtr<FUICommandList>PaletteActionsCommandList;
+};
+
+/** An FToolPalette to which you can add and remove actions */
+struct WIDGETREGISTRATION_API FEditablePalette : public FToolPalette
+{
+	FEditablePalette(TSharedPtr<FUICommandInfo> InLoadToolPaletteAction,
+		TArray<FString>& InPaletteCommandNameArray,
+		TSharedPtr<FUICommandInfo> InAddToPaletteAction,
+		TSharedPtr<FUICommandInfo> InRemoveFromPaletteAction);
+
+	/** The FUICommandInfo which adds an action to this palette */
+	const TSharedPtr<FUICommandInfo> AddToPaletteAction;
+	
+	/** The FUICommandInfo which removes an action to this palette */
+	const TSharedPtr<FUICommandInfo> RemoveFromPaletteAction;
+	
+	/** The TArray of Command names that are the current FuiCommandInfo actions in this Palette */
+	TArray<FString>& PaletteCommandNameArray;
+
+	/**
+	 * Given a reference to a FUICommandInfo, returns whether it is in the current Palette
+	 *
+	 * @param CommandName the name of the FUICommandInfo queried as to whether it is in the Palette
+	 */
+	bool IsInPalette(const FName CommandName) const;
+};
+
+
+/**
+ * The FToolElementRegistrationArgs which is specified for Toolkits
+ */
+class WIDGETREGISTRATION_API FToolkitBuilder : public FToolElementRegistrationArgs
+{
+public:
+	FToolkitBuilder(
+		FName InToolbarCustomizationName,
+		TSharedPtr<FUICommandList> InToolkitCommandList);
+
+	/** A delegate for any clean up that may be needed when the load tool palette FUICommandAction ends/unloads */
+	DECLARE_DELEGATE(FOnToolEnded);
+	FOnToolEnded OnToolEnded;
+
+	/**
+	 * Adds the FToolPalette Palette to this FToolkitBuilder
+	 *
+	 * @param Palette the FToolPalette being added to this FToolkitBuilder
+	 */
+	void AddPalette(TSharedPtr<FToolPalette> Palette);
+
+	/**
+	 * Adds the FEditablePalette Palette to this FToolkitBuilder
+	 *
+	 * @param Palette the FEditablePalette being added to this FToolkitBuilder
+	 */
+	void AddPalette(TSharedPtr<FEditablePalette> Palette);
+
+	/** Cleans up any previously set data in this FToolkitBuilder and reset the members to their initial values  */
+	virtual void ResetWidget() override;
+
+	/** Updates the Toolkit. This should be called after any changes to the data of this FToolkitBuilder, and the UI
+	 * will be regenerated to reflect it.  */
+	virtual void UpdateWidget() override;
+
+	/** Creates the Toolbar for the widget with the FUICommandInfos that load the Palettes */
+	TSharedRef<SWidget> CreateToolbarWidget() const;
+
+	/** returns a TSharedPointer to the FToolbarBuilder with the FUICommandInfos that load the Palettes */
+	TSharedPtr<FToolBarBuilder> GetLoadPaletteToolbar();
+	
+	/** returns a TSharedPointer to the FToolbarBuilder with the FUICommandInfos that load the Palettes */
+	TSharedPtr<FToolElement> VerticalToolbarElement;
+
+	/** returns a TSharedPointer to the FToolbarBuilder with the FUICommandInfos that load the Palettes */
+	TSharedRef<SWidget> GetToolPaletteWidget() const;
+
+	/** returns true is there is an active palette selected, else it returns false */
+	bool HasActivePalette() const;
+
+	/**
+	 * Returns true if the FUICommandInfo with the name CommandName is the active tool palette,
+	 * else it returns false
+	 *
+	 * @param CommandName the name of the FUICommandInfo we are checking to see if it is the active tool palette
+	 */
+	ECheckBoxState IsActiveToolPalette(FName CommandName) const;
+	
+private:
+
+	/** the tool element registry this class will use to register UI tool elements */
+	static FToolElementRegistry ToolRegistry;
+	
+	/** Name of the toolbar this mode uses and can be used by external systems to customize that mode toolbar */
+	FName ToolbarCustomizationName;
+	TMap<FString, TSharedPtr<FButtonArgs>> PaletteCommandNameToButtonArgsMap;
+
+	/** A TSharedPointer to the FUICommandList for the FUICommandInfos which load a tool palette */
+	TSharedPtr<FUICommandList> LoadToolPaletteCommandList;
+
+	/** A TSharedPointer to the FUICommandList for the current mode */
+	TSharedPtr<FUICommandList> ToolkitCommandList;
+
+	/** A TArray of FEditablePalettes, kept to update the commands which are on them */
+	TArray<TSharedRef<FEditablePalette>> EditablePalettesArray;
+
+	/** The tool palette which is currently loaded/active */
+	TSharedPtr<FToolPalette> ActivePalette;
+
+	/** The SVerticalBox which holds the tool palette (the Buttons which load each tool) */
+	TSharedPtr<SVerticalBox> ToolPaletteWidget;
+
+	/** The toolbar builder for the toolbar which has the FUICommandInfos which load the various palettes */
+	TSharedPtr<FToolBarBuilder> LoadPaletteToolBarBuilder;
+
+	/** the FName of each Load palette command to the FToolbarBuilder for the palette which it loads */
+	TMap<FName, TSharedPtr<FToolBarBuilder>> LoadCommandNameToPaletteToolbarBuilderMap;
+
+	/** Updates the Editable Palette with any commands that are on it */
+	void UpdateEditablePalette(FEditablePalette& EditablePalette);
+
+	/** Resets the tool palette widget, on which the buttons for the currently chosen toolset are shown */
+	void ResetToolPaletteWidget();
+
+	/**
+	 * Creates the UI for the Palette specified by the FToolPalette Palette
+	 *
+	 * @param Palette the FToolPalette for which this method will build the UI 
+	 */
+	void CreatePalette(TSharedPtr<FToolPalette> Palette);
+
+	/**
+	 * Creates the UI for the Palette specified by the FToolPalette Palette if it is not currently active,
+	 * else it removes it.
+	 *
+	 * @param Palette the FToolPalette for which this method will build the UI 
+	 */
+	void TogglePalette(TSharedPtr<FToolPalette> Palette);
+
+	/**
+	 * Adds the FUICommandInfo with the command name CommandNameString to the FEditablePalette Palette if
+	 * it is not on Palette, and removes the FUICommandInfo from Palette if it is on the Palette
+	 *
+	 * @param Palette the FEditablePalette to which the FUICommandInfo will be toggled on/off
+	 * @param CommandNameString the name of the FUICommandInfo command which will be toggled on/off of Palette
+	 */
+	void ToggleCommandInPalette(TSharedRef<FEditablePalette> Palette, FString CommandNameString);
+
+	/**
+	 * Retrieves the context menu content for the FUICommandInfo with the command name CommandName
+	 *
+	 * @param CommandName the command name of the FUICommandInfo to get the context menu widget for
+	 * @return the TSharedRef<SWidget> which contains the context menu for the FUICommandInfo with the name CommandName 
+	 */
+	TSharedRef<SWidget> GetContextMenuContent(const FName CommandName);
+};
\ No newline at end of file
diff --git a/Engine/Source/Developer/WidgetRegistration/WidgetRegistration.Build.cs b/Engine/Source/Developer/WidgetRegistration/WidgetRegistration.Build.cs
new file mode 100644
--- /dev/null
+++ b/Engine/Source/Developer/WidgetRegistration/WidgetRegistration.Build.cs
@@ -0,0 +1,22 @@
+// Copyright Epic Games, Inc. All Rights Reserved.
+
+using UnrealBuildTool;
+
+public class WidgetRegistration : ModuleRules
+{
+	public WidgetRegistration(ReadOnlyTargetRules Target) : base(Target)
+	{
+		
+		PublicDependencyModuleNames.AddRange(
+			new string[] {
+				"SlateCore",
+				"Core",
+				"Slate",
+				"Engine",
+				"ToolMenus"
+			});
+		
+
+			
+	}
+}
\ No newline at end of file
diff --git a/Engine/Source/Editor/EditorStyle/Private/StarshipStyle.cpp b/Engine/Source/Editor/EditorStyle/Private/StarshipStyle.cpp
--- a/Engine/Source/Editor/EditorStyle/Private/StarshipStyle.cpp
+++ b/Engine/Source/Editor/EditorStyle/Private/StarshipStyle.cpp
@@ -1670,6 +1670,90 @@ void FStarshipEditorStyle::FStyle::SetupGeneralStyles()
 		Set("EditorModesPanel.ToolDescriptionFont", DEFAULT_FONT("Italic", 10));
 
 	}
+
+	{
+
+		
+		FToolBarStyle SlimPaletteToolBarStyle = FStarshipCoreStyle::GetCoreStyle().GetWidgetStyle<FToolBarStyle>("SlimToolBar");
+
+		FTextBlockStyle SlimPaletteToolbarLabelStyle = FTextBlockStyle(GetParentStyle()->GetWidgetStyle<FTextBlockStyle>("NormalText"));
+		SlimPaletteToolbarLabelStyle.SetOverflowPolicy(ETextOverflowPolicy::Ellipsis);
+		SlimPaletteToolBarStyle.SetLabelStyle(SlimPaletteToolbarLabelStyle);
+		SlimPaletteToolBarStyle.SetBackground(FSlateColorBrush(FStyleColors::Recessed));
+		SlimPaletteToolBarStyle.SetLabelPadding(FMargin(0.0f, 0.0f, 0.0f, 0.0f));
+		SlimPaletteToolBarStyle.SetIconSize(Icon20x20);
+
+		SlimPaletteToolBarStyle.SetButtonPadding(       FMargin(0.0f, 4.0f, 0.0f, 0.0f));
+		SlimPaletteToolBarStyle.SetCheckBoxPadding(     FMargin(0.0f, 0.0f));
+		// SlimPaletteToolBarStyle.SetIconPadding( FMargin(8.f, 4.f, 8.f, 4.f) ); @TODO: ~ enable with styling changesS
+		SlimPaletteToolBarStyle.SetComboButtonPadding(  FMargin(0.0f, 0.0f));
+		SlimPaletteToolBarStyle.SetIndentedBlockPadding(FMargin(0.0f, 0.0f));
+		SlimPaletteToolBarStyle.SetBlockPadding(        FMargin(0.0f, 0.0f));
+		SlimPaletteToolBarStyle.ToggleButton.SetPadding(FMargin(0.0f, 1.0f));
+		SlimPaletteToolBarStyle.ButtonStyle.SetNormalPadding(FMargin(2.0f, 1.0f));
+		SlimPaletteToolBarStyle.ButtonStyle.SetPressedPadding(FMargin(2.0f, 1.0f));
+
+		FCheckBoxStyle CheckBoxStyle = FCheckBoxStyle(SlimPaletteToolBarStyle.ToggleButton)
+			.SetCheckedImage(FSlateRoundedBoxBrush(FStyleColors::Primary, 4.f, FLinearColor(0, 0, 0, .8), 0.5))
+			.SetCheckedHoveredImage(FSlateRoundedBoxBrush(FStyleColors::PrimaryHover, 4.f, FLinearColor(0, 0, 0, .8), 0.5))
+			.SetCheckedPressedImage(FSlateRoundedBoxBrush(FStyleColors::PrimaryPress, 4.f, FLinearColor(0, 0, 0, .8), 0.5))
+			.SetUncheckedImage(FSlateRoundedBoxBrush(FStyleColors::Dropdown, 4.f, FLinearColor(0, 0, 0, .8), 0.5))
+			.SetForegroundColor(FStyleColors::Foreground)
+			.SetPressedForegroundColor(FStyleColors::ForegroundHover)
+			.SetHoveredForegroundColor(FStyleColors::ForegroundHover)
+			.SetCheckedForegroundColor(FStyleColors::Foreground)
+			.SetCheckedPressedForegroundColor(FStyleColors::ForegroundHover)
+			.SetCheckedHoveredForegroundColor(FStyleColors::ForegroundHover)
+			.SetPadding(FMargin(0.f, 0.f, 0.f, 0.f));
+
+		Set("SlimPaletteToolBarStyle.ToggleButton", CheckBoxStyle);
+		Set("SlimPaletteToolBarStyle.TitleBar.Padding", FMargin(4.f, 4.f));
+		SlimPaletteToolBarStyle.SetToggleButtonStyle(CheckBoxStyle);
+		Set("SlimPaletteToolBar", SlimPaletteToolBarStyle);
+		Set("ToolAreaBackground",  new FSlateRoundedBoxBrush(FStyleColors::AccentGreen, 4.0f, FLinearColor(0, 0, 0, .8), 0.5));
+		Set( "SlimPaletteToolBar.PaletteButtonBorderBrush", new FSlateRoundedBoxBrush(FStyleColors::AccentGreen, 4.0f, FLinearColor(0, 0, 0, .8), 0.5) );
+	}
+
+{
+		FToolBarStyle FVerticalToolBarStyle = FStarshipCoreStyle::GetCoreStyle().GetWidgetStyle<FToolBarStyle>("SlimToolBar");
+
+		FTextBlockStyle FVerticalToolBarLabelStyle = FTextBlockStyle(GetParentStyle()->GetWidgetStyle<FTextBlockStyle>("SmallText"));
+		FVerticalToolBarLabelStyle.SetOverflowPolicy(ETextOverflowPolicy::Ellipsis);
+		
+		FVerticalToolBarStyle.SetLabelStyle(FVerticalToolBarLabelStyle);
+		FVerticalToolBarStyle.SetLabelPadding(FMargin(0.0f, 1.0f, 0.0f, 0.0f));
+
+		FVerticalToolBarStyle.SetButtonPadding(       FMargin(0.0f, 0.0f));
+		FVerticalToolBarStyle.SetCheckBoxPadding(     FMargin(0.0f, 0.0f));
+		FVerticalToolBarStyle.SetComboButtonPadding(  FMargin(0.0f, 0.0f));
+		FVerticalToolBarStyle.SetIndentedBlockPadding(FMargin(0.0f, 0.0f));
+		FVerticalToolBarStyle.SetBlockPadding(        FMargin(0.0f, 0.0f));
+		FVerticalToolBarStyle.SetBackgroundPadding(   FMargin(4.0f, 2.0f));
+		FVerticalToolBarStyle.ButtonStyle.SetNormalPadding(FMargin(2.0f, 6.0f));
+		FVerticalToolBarStyle.ButtonStyle.SetPressedPadding(FMargin(2.0f, 6.0f));
+		FVerticalToolBarStyle.SetExpandBrush(IMAGE_BRUSH("Icons/toolbar_expand_16x", Icon8x8));
+ 
+		FCheckBoxStyle CheckBoxStyle = FCheckBoxStyle(FVerticalToolBarStyle.ToggleButton)
+			.SetCheckedImage(FSlateRoundedBoxBrush(FStyleColors::Primary, 4.f, FLinearColor(0, 0, 0, .8), 0.5))
+			.SetCheckedHoveredImage(FSlateRoundedBoxBrush(FStyleColors::PrimaryHover, 4.f, FLinearColor(0, 0, 0, .8), 0.5))
+			.SetCheckedPressedImage(FSlateRoundedBoxBrush(FStyleColors::PrimaryPress, 4.f, FLinearColor(0, 0, 0, .8), 0.5))
+			.SetForegroundColor(FStyleColors::Foreground)
+			.SetPressedForegroundColor(FStyleColors::ForegroundHover)
+			.SetHoveredForegroundColor(FStyleColors::ForegroundHover)
+			.SetCheckedForegroundColor(FStyleColors::Foreground)
+			.SetCheckedPressedForegroundColor(FStyleColors::ForegroundHover)
+			.SetCheckedHoveredForegroundColor(FStyleColors::ForegroundHover)
+			.SetPadding(FMargin(6.0f, 4.0f, 6.0f, 4.0f));
+
+		Set("FVerticalToolBar.ToggleButton", CheckBoxStyle);
+		FVerticalToolBarStyle.SetToggleButtonStyle(CheckBoxStyle);
+
+		Set("FVerticalToolBar.MaxUniformToolbarSize", 48.f);
+		Set("FVerticalToolBar.MinUniformToolbarSize", 48.f);
+
+		Set("FVerticalToolBar", FVerticalToolBarStyle);
+	}
+	
 	// Vertical ToolPalette 
 	{
 		FToolBarStyle VerticalToolBarStyle = FStarshipCoreStyle::GetCoreStyle().GetWidgetStyle<FToolBarStyle>("SlimToolBar");
@@ -1715,9 +1799,6 @@ void FStarshipEditorStyle::FStyle::SetupGeneralStyles()
 		Set("VerticalToolBar.MaxUniformToolbarSize", 48.f);
 		Set("VerticalToolBar.MinUniformToolbarSize", 48.f);
 
-		Set("VerticalToolBar.ExpandableAreaHeader", new FSlateRoundedBoxBrush(FStyleColors::Dropdown, FVector4(4.0, 4.0, 0.0, 0.0)));
-		Set("VerticalToolBar.ExpandableAreaBody", new FSlateRoundedBoxBrush(FStyleColors::Recessed, FVector4(0.0, 0.0, 4.0, 4.0)));
-
 		Set("VerticalToolBar", VerticalToolBarStyle);
 	}
 	
@@ -4517,7 +4598,7 @@ void FStarshipEditorStyle::FStyle::SetupLevelEditorStyle()
 	{
 
 		Set( "LevelViewport.DebugBorder", new BOX_BRUSH( "Old/Window/ViewportDebugBorder", 0.8f, FLinearColor(.7,0,0,.5) ) );
-		Set( "LevelViewport.BlackBackground", new FSlateColorBrush( FLinearColor::Black ) ); 
+		Set( "LevelViewport.BlackBackground", new FSlateColorBrush( FLinearColor::Red ) ); 
 		Set( "LevelViewport.StartingPlayInEditorBorder", new BOX_BRUSH( "Old/Window/ViewportDebugBorder", 0.8f, FLinearColor(0.1f,1.0f,0.1f,1.0f) ) );
 		Set( "LevelViewport.StartingSimulateBorder", new BOX_BRUSH( "Old/Window/ViewportDebugBorder", 0.8f, FLinearColor(1.0f,1.0f,0.1f,1.0f) ) );
 		Set( "LevelViewport.NonMaximizedBorder", new CORE_BORDER_BRUSH("Common/PlainBorder", 2.f / 8.f, FStyleColors::Black));
diff --git a/Engine/Source/Editor/UnrealEd/Private/Toolkits/BaseToolkit.cpp b/Engine/Source/Editor/UnrealEd/Private/Toolkits/BaseToolkit.cpp
--- a/Engine/Source/Editor/UnrealEd/Private/Toolkits/BaseToolkit.cpp
+++ b/Engine/Source/Editor/UnrealEd/Private/Toolkits/BaseToolkit.cpp
@@ -21,6 +21,7 @@
 #include "Widgets/Images/SImage.h"
 #include "Widgets/Input/SSegmentedControl.h"
 #include "Styling/AppStyle.h"
+#include "Styling/StyleColors.h"
 
 #define LOCTEXT_NAMESPACE "BaseToolkit"
 
@@ -180,6 +181,11 @@ void FModeToolkit::Init(const TSharedPtr< class IToolkitHost >& InitToolkitHost,
 	}
 
 	GetEditorModeManager().OnEditorModeIDChanged().AddSP(this, &FModeToolkit::OnModeIDChanged);
+
+	if (HasToolkitBuilder())	
+	{
+		ToolkitBuilder->VerticalToolbarElement->GenerateWidget();
+	}
 }
 
 
@@ -445,10 +451,17 @@ void FModeToolkit::RebuildModeToolPalette()
 
 }
 
+bool FModeToolkit::HasToolkitBuilder() const
+{
+	return bUsesToolkitBuilder && ToolkitBuilder != nullptr;
+}
 
 TSharedRef<SDockTab> FModeToolkit::CreatePrimaryModePanel(const FSpawnTabArgs& Args)
 {
-	TSharedRef<SWidget> TabContent = SNew(SBorder)
+	TSharedPtr<SWidget> TabContent;
+	if (!HasToolkitBuilder())
+	{
+		TabContent = SNew(SBorder)
 		.BorderImage(FAppStyle::GetBrush("ToolPanel.GroupBorder"))
 		.Padding(0.0f)
 		[
@@ -484,13 +497,31 @@ TSharedRef<SDockTab> FModeToolkit::CreatePrimaryModePanel(const FSpawnTabArgs& A
 				]
 		]
 		];
-
-
-	TSharedPtr<SDockTab> CreatedTab = SNew(SDockTab)
+	}
+	// else if ToolkitBuilder is defined, make the Toolkit
+	else
+	{
+		TabContent = SAssignNew(InlineContentHolder, SBorder)
+		.BorderImage(FAppStyle::GetBrush("ToolPanel.GroupBorder"))
+			.Padding(0.0f)
 		[
-			TabContent
+			SNew(SVerticalBox)
+			+ SVerticalBox::Slot()
+			.AutoHeight()
+			.HAlign(HAlign_Left)
+			[
+				SAssignNew(ModeToolBarContainer, SBorder)
+				.Padding(FMargin(4, 0, 0, 0))
+				.BorderImage(FAppStyle::GetBrush("ToolPanel.GroupBorder"))
+			]
 		];
+	}
 
+	const TSharedPtr<SDockTab> CreatedTab = SNew(SDockTab);
+	if (TabContent)
+	{
+		CreatedTab->SetContent(TabContent.ToSharedRef());
+	}
 	PrimaryTab = CreatedTab;	
 	UpdatePrimaryModePanel();
 	return CreatedTab.ToSharedRef();
@@ -507,8 +538,15 @@ void FModeToolkit::UpdatePrimaryModePanel()
 		const FSlateBrush* TabIcon = GetEditorModeIcon().GetSmallIcon();
 		if (PrimaryTab.IsValid())
 		{
-			PrimaryTab.Pin()->SetTabIcon(TabIcon);
-			PrimaryTab.Pin()->SetLabel(TabName);
+			TSharedPtr<SDockTab> TabPtr  = PrimaryTab.Pin(); 
+			TabPtr->SetTabIcon(TabIcon);
+			TabPtr->SetLabel(TabName);
+
+			if (HasToolkitBuilder())
+			{
+				TabPtr->SetParentDockTabStackTabWellHidden(HasToolkitBuilder());
+				return;				
+			}
 		}
 
 		if (HasIntegratedToolPalettes())
@@ -654,7 +692,7 @@ void FModeToolkit::RequestModeUITabs()
 		PrimaryTabInfo.TabLabel = LOCTEXT("ModesToolboxTab", "Mode Toolbox");
 		PrimaryTabInfo.TabTooltip = LOCTEXT("ModesToolboxTabTooltipText", "Open the  Modes tab, which contains the active editor mode's settings.");
 		ModeUILayerPtr->SetModePanelInfo(UAssetEditorUISubsystem::TopLeftTabID, PrimaryTabInfo);
-		if (!HasIntegratedToolPalettes())
+		if (!HasIntegratedToolPalettes() && !HasToolkitBuilder())
 		{
 
 			ToolbarInfo.OnSpawnTab = FOnSpawnTab::CreateSP(SharedThis(this), &FModeToolkit::MakeModeToolbarTab);
@@ -729,9 +767,15 @@ void FModeToolkit::SpawnOrUpdateModeToolbar()
 
 void FModeToolkit::RebuildModeToolBar()
 {
+	TSharedPtr<SDockTab> ToolbarTabPtr = ModeToolbarTab.Pin();
+	if (ToolbarTabPtr  && HasToolkitBuilder())
+	{
+		ToolbarTabPtr->SetParentDockTabStackTabWellHidden(true);
+	}
+
 	// If the tab or box is not valid the toolbar has not been opened or has been closed by the user
 	TSharedPtr<SVerticalBox> ModeToolbarBoxPinned = ModeToolbarBox.Pin();
-	if (ModeToolbarTab.IsValid() && ModeToolbarBoxPinned.IsValid())
+	if (ModeToolbarTab.IsValid() && ModeToolbarBoxPinned)
 	{
 		ModeToolbarBoxPinned->ClearChildren();
 		bool bExclusivePalettes = true;
@@ -782,12 +826,12 @@ void FModeToolkit::RebuildModeToolBar()
 								[
 									SNew(SCheckBox)
 									.Style(FAppStyle::Get(), "ToolPalette.DockingTab")
-								.OnCheckStateChanged_Lambda([PaletteSwitcher, Row, this](const ECheckBoxState) {
+									.OnCheckStateChanged_Lambda([PaletteSwitcher, Row, this](const ECheckBoxState) {
 										PaletteSwitcher->SetActiveWidget(Row.ToolbarWidget.ToSharedRef());
 										SetCurrentPalette(Row.PaletteName);
 									}
 								)
-								.IsChecked_Lambda([PaletteSwitcher, PaletteWidget]() -> ECheckBoxState { return PaletteSwitcher->GetActiveWidget() == PaletteWidget ? ECheckBoxState::Checked : ECheckBoxState::Unchecked; })
+									.IsChecked_Lambda([PaletteSwitcher, PaletteWidget]() -> ECheckBoxState { return PaletteSwitcher->GetActiveWidget() == PaletteWidget ? ECheckBoxState::Checked : ECheckBoxState::Unchecked; })
 										[
 											SNew(STextBlock)
 											.Text(Row.DisplayName)
diff --git a/Engine/Source/Editor/UnrealEd/Public/Toolkits/BaseToolkit.h b/Engine/Source/Editor/UnrealEd/Public/Toolkits/BaseToolkit.h
--- a/Engine/Source/Editor/UnrealEd/Public/Toolkits/BaseToolkit.h
+++ b/Engine/Source/Editor/UnrealEd/Public/Toolkits/BaseToolkit.h
@@ -3,6 +3,7 @@
 #pragma once
 
 #include "CoreMinimal.h"
+#include "ToolkitBuilder.h"
 #include "Framework/Docking/WorkspaceItem.h"
 #include "Framework/Docking/TabManager.h"
 #include "Toolkits/IToolkit.h"
@@ -259,4 +260,14 @@ protected:
 
 	/** All toolbar rows generated by active modes.  There will be one row per active mode that generates a toolbar */
 	TArray<FEdModeToolbarRow> ActiveToolBarRows;
+
+	/** a feature flag that can be set to true in order to make and use a ToolkitBuilder,
+	 *  if one is defined for a given mode. */ 
+	bool bUsesToolkitBuilder = false;
+	
+	/** The FToolkitBuilder which will build the UI for this mode, if defined */
+	TSharedPtr<FToolkitBuilder> ToolkitBuilder;
+
+	/** returns whether or not this FModeToolkit has a ToolkitBuilder defined */
+	FORCEINLINE bool HasToolkitBuilder() const;
 };
diff --git a/Engine/Source/Editor/UnrealEd/UnrealEd.Build.cs b/Engine/Source/Editor/UnrealEd/UnrealEd.Build.cs
--- a/Engine/Source/Editor/UnrealEd/UnrealEd.Build.cs
+++ b/Engine/Source/Editor/UnrealEd/UnrealEd.Build.cs
@@ -72,6 +72,7 @@ public class UnrealEd : ModuleRules
 				"WorkspaceMenuStructure",
 				"LandscapeEditor",
 				"Blutility",
+				"WidgetRegistration",
 				"SlateReflector"
 			}
 		);
@@ -122,7 +123,8 @@ public class UnrealEd : ModuleRules
 				"SubobjectEditor",
 				"PhysicsUtilities",
 				"ToolWidgets",
-				"AssetTools",
+				"AssetTools", 
+				"WidgetRegistration"
 			}
 		);
 
diff --git a/Engine/Source/Runtime/Slate/Private/Framework/MultiBox/MultiBox.cpp b/Engine/Source/Runtime/Slate/Private/Framework/MultiBox/MultiBox.cpp
--- a/Engine/Source/Runtime/Slate/Private/Framework/MultiBox/MultiBox.cpp
+++ b/Engine/Source/Runtime/Slate/Private/Framework/MultiBox/MultiBox.cpp
@@ -1,6 +1,8 @@
 // Copyright Epic Games, Inc. All Rights Reserved.
 
 #include "Framework/MultiBox/MultiBox.h"
+
+#include "SClippingVerticalBox.h"
 #include "Misc/ConfigCacheIni.h"
 #include "Widgets/SBoxPanel.h"
 #include "Widgets/SOverlay.h"
@@ -252,7 +254,7 @@ FMultiBox::FMultiBox(const EMultiBoxType InType, FMultiBoxCustomization InCustom
 	, bShouldCloseWindowAfterMenuSelection( bInShouldCloseWindowAfterMenuSelection )
 {
 
-	if (InType == EMultiBoxType::SlimHorizontalToolBar && FCoreStyle::IsStarshipStyle())
+	if ((InType == EMultiBoxType::SlimHorizontalToolBar ||  InType == EMultiBoxType::SlimHorizontalUniformToolBar) && FCoreStyle::IsStarshipStyle())
 	{
 		StyleName = "SlimToolBar";
 	}
@@ -669,7 +671,7 @@ void SMultiBoxWidget::AddBlockWidget(const FMultiBlock& Block, TSharedPtr<SHoriz
 	check( MultiBox.IsValid() );
 
 	// Skip Separators for Uniform Tool Bars
-	if ( Block.GetType() == EMultiBlockType::Separator  && MultiBox->GetType() == EMultiBoxType::UniformToolBar)
+	if ( Block.GetType() == EMultiBlockType::Separator  && (MultiBox->GetType() == EMultiBoxType::UniformToolBar || MultiBox->GetType() == EMultiBoxType::SlimHorizontalUniformToolBar))
 	{
 		return;
 	}
@@ -823,6 +825,7 @@ void SMultiBoxWidget::AddBlockWidget(const FMultiBlock& Block, TSharedPtr<SHoriz
 		}
 		break;
 	case EMultiBoxType::UniformToolBar:
+	case EMultiBoxType::SlimHorizontalUniformToolBar:
 		{
 			UniformToolbarPanel->AddSlot()
 			[
@@ -971,7 +974,11 @@ void SMultiBoxWidget::BuildMultiBoxWidget()
 		break;
 	case EMultiBoxType::VerticalToolBar:
 		{
-			MainWidget = VerticalBox = SNew(SVerticalBox);
+			MainWidget = VerticalBox = ClippedVerticalBox = SNew(SClippingVerticalBox)
+				.OnWrapButtonClicked(FOnGetContent::CreateSP(this, &SMultiBoxWidget::OnWrapButtonClicked))
+				.IsFocusable(MultiBox->bIsFocusable)
+				.StyleSet(StyleSet)
+				.StyleName(StyleName);
 		}
 		break;
 	case EMultiBoxType::UniformToolBar:
@@ -992,6 +999,26 @@ void SMultiBoxWidget::BuildMultiBoxWidget()
 			];
 		}
 		break;
+		// @TODO: ~Move hardcodes into styling file
+	case EMultiBoxType::SlimHorizontalUniformToolBar:
+		{
+			MainWidget = VerticalBox = SNew (SVerticalBox)
+
+			+SVerticalBox::Slot()
+			.Padding(FMargin(0.f, 2.f))
+			.AutoHeight()
+			[
+				SAssignNew(UniformToolbarPanel, SUniformWrapPanel)
+				.HAlign(HAlign_Fill)
+				.MinDesiredSlotWidth(150.f)
+				.MinDesiredSlotHeight(32.f)
+				.MaxDesiredSlotWidth(150.f)
+				.MaxDesiredSlotHeight(32.f)
+				.NumColumnsOverride(2)
+				.SlotPadding(FMargin(2.f, 1.f))
+			];
+		}
+		break;
 	case EMultiBoxType::ButtonRow:
 		{
 			MainWidget = TileView = SNew(STileView< TSharedPtr<SWidget> >)
@@ -1182,6 +1209,11 @@ void SMultiBoxWidget::BuildMultiBoxWidget()
 	{
 		ClippedHorizontalBox->AddWrapButton();
 	}
+	
+	if (ClippedVerticalBox.IsValid())
+	{
+		ClippedVerticalBox->AddWrapButton();
+	}
 
 	FMargin BorderPadding(0);
 	const FSlateBrush* BorderBrush = FStyleDefaults::GetNoBrush();
@@ -1234,7 +1266,9 @@ TSharedRef<SWidget> SMultiBoxWidget::OnWrapButtonClicked()
 	FMenuBuilder MenuBuilder(true, MultiBox->GetLastCommandList(), TSharedPtr<FExtender>(), false, GetStyleSet());
 	{ 
 		const TArray< TSharedRef< const FMultiBlock > >& Blocks = MultiBox->GetBlocks();
-		for (int32 BlockIdx = ClippedHorizontalBox->GetClippedIndex(); BlockIdx < Blocks.Num(); ++BlockIdx)
+
+		int32 BlockIdx = ClippedHorizontalBox ? ClippedHorizontalBox->GetClippedIndex() : ClippedVerticalBox->GetClippedIndex();
+		for (; BlockIdx < Blocks.Num(); ++BlockIdx)
 		{
 			Blocks[BlockIdx]->CreateMenuEntry(MenuBuilder);
 		}
diff --git a/Engine/Source/Runtime/Slate/Private/Framework/MultiBox/MultiBoxBuilder.cpp b/Engine/Source/Runtime/Slate/Private/Framework/MultiBox/MultiBoxBuilder.cpp
--- a/Engine/Source/Runtime/Slate/Private/Framework/MultiBox/MultiBoxBuilder.cpp
+++ b/Engine/Source/Runtime/Slate/Private/Framework/MultiBox/MultiBoxBuilder.cpp
@@ -20,6 +20,7 @@
 #include "Framework/MultiBox/ToolMenuBase.h"
 #include "Widgets/Layout/SScrollBox.h"
 #include "Styling/ToolBarStyle.h"
+#include "Framework/MultiBox/SUniformToolbarButtonBlock.h"
 
 
 FMultiBoxBuilder::FMultiBoxBuilder( const EMultiBoxType InType, FMultiBoxCustomization InCustomization, const bool bInShouldCloseWindowAfterMenuSelection, const TSharedPtr< const FUICommandList >& InCommandList, TSharedPtr<FExtender> InExtender, FName InTutorialHighlightName, FName InMenuName )
@@ -431,6 +432,7 @@ void FToolBarBuilder::SetIsFocusable(bool bInIsFocusable)
 	MultiBox->bIsFocusable = bIsFocusable;
 }
 
+
 void FToolBarBuilder::AddToolBarButton(const TSharedPtr< const FUICommandInfo > InCommand, FName InExtensionHook, const TAttribute<FText>& InLabelOverride, const TAttribute<FText>& InToolTipOverride, const TAttribute<FSlateIcon>& InIconOverride, FName InTutorialHighlightName, FNewMenuDelegate InCustomMenuDelegate )
 {
 	ApplySectionBeginning();
@@ -455,6 +457,42 @@ void FToolBarBuilder::AddToolBarButton(const TSharedPtr< const FUICommandInfo >
 	ApplyHook(InExtensionHook, EExtensionHook::After);
 }
 
+void FToolBarBuilder::AddToolBarButton(FButtonArgs& Args)
+{
+	ApplySectionBeginning();
+
+	ApplyHook(Args.ExtensionHook, EExtensionHook::Before);
+
+	const bool bHasUserInteractionType = Args.UserInterfaceActionType != EUserInterfaceActionType::None;
+	
+	const TSharedPtr< FToolBarButtonBlock > NewToolBarButtonBlock( bHasUserInteractionType ?
+		new FToolBarButtonBlock( Args.LabelOverride, Args.ToolTipOverride, Args.IconOverride, Args.Action, Args.UserInterfaceActionType ) :
+		new FToolBarButtonBlock( Args.Command.ToSharedRef(), CommandListStack.Last(), Args.LabelOverride, Args.ToolTipOverride, Args.IconOverride )
+		);
+
+	if ( LabelVisibility.IsSet() )
+	{
+		NewToolBarButtonBlock->SetLabelVisibility( LabelVisibility.GetValue() );
+	}
+
+	NewToolBarButtonBlock->SetIsFocusable(bIsFocusable);
+	NewToolBarButtonBlock->SetForceSmallIcons(bForceSmallIcons);
+	NewToolBarButtonBlock->SetTutorialHighlightName(GenerateTutorialIdentfierName(TutorialHighlightName, Args.TutorialHighlightName, Args.Command, MultiBox->GetBlocks().Num()));
+	NewToolBarButtonBlock->SetStyleNameOverride(CurrentStyleOverride);
+	NewToolBarButtonBlock->SetCustomMenuDelegate(Args.CustomMenuDelegate);
+	NewToolBarButtonBlock->SetOnGetMenuContent(Args.OnGetMenuContent);
+	
+	if (bHasUserInteractionType)
+	{
+		NewToolBarButtonBlock->SetStyleNameOverride(CurrentStyleOverride);
+		NewToolBarButtonBlock->SetCustomMenuDelegate(Args.CustomMenuDelegate);
+	}
+
+	MultiBox->AddMultiBlock( NewToolBarButtonBlock.ToSharedRef() );
+
+	ApplyHook(Args.ExtensionHook, EExtensionHook::After);
+}
+
 void FToolBarBuilder::AddToolBarButton(const FUIAction& InAction, FName InExtensionHook, const TAttribute<FText>& InLabelOverride, const TAttribute<FText>& InToolTipOverride, const TAttribute<FSlateIcon>& InIconOverride, const EUserInterfaceActionType UserInterfaceActionType, FName InTutorialHighlightName )
 {
 	ApplySectionBeginning();
@@ -697,3 +735,38 @@ void FButtonRowBuilder::AddButton( const FText& InLabel, const FText& InToolTip,
 	MultiBox->AddMultiBlock( NewButtonRowBlock );
 }
 
+FSlimHorizontalUniformToolBarBuilder::FSlimHorizontalUniformToolBarBuilder(TSharedPtr<const FUICommandList> InCommandList, FMultiBoxCustomization InCustomization, TSharedPtr<FExtender> InExtender, const bool InForceSmallIcons)
+	: FToolBarBuilder(EMultiBoxType::SlimHorizontalUniformToolBar, InCommandList, InCustomization, InExtender, InForceSmallIcons)
+{
+	const FToolBarStyle& ToolBarStyle = GetStyleSet()->GetWidgetStyle<FToolBarStyle>(GetStyleName());
+}
+
+void FSlimHorizontalUniformToolBarBuilder::AddToolBarButton( FButtonArgs& ButtonArgs )
+{
+ 	ApplySectionBeginning();
+	ApplyHook(ButtonArgs.ExtensionHook, EExtensionHook::Before);
+
+//	ButtonArgs.BorderBrushName = "SlimPaletteToolBar.PaletteButtonBorderBrush";
+	const TSharedPtr< FToolBarButtonBlock > NewHorizontalToolBarButtonBlock( new FToolBarButtonBlock( ButtonArgs) );
+    InitializeToolBarButtonBlock(StaticCastSharedPtr<FToolBarButtonBlock>(NewHorizontalToolBarButtonBlock), ButtonArgs);
+}
+void FToolBarBuilder::InitializeToolBarButtonBlock(TSharedPtr<FToolBarButtonBlock> ToolBarButtonBlock, FButtonArgs& ButtonArgs)
+{
+	if ( LabelVisibility.IsSet() )
+	{
+		ToolBarButtonBlock->SetLabelVisibility( LabelVisibility.GetValue() );
+	}
+
+	ToolBarButtonBlock->SetBorderBrushName(ButtonArgs.BorderBrushName);
+	ToolBarButtonBlock->SetIsFocusable(bIsFocusable);
+	ToolBarButtonBlock->SetForceSmallIcons(bForceSmallIcons);
+	ToolBarButtonBlock->SetTutorialHighlightName(GenerateTutorialIdentfierName(
+	TutorialHighlightName, ButtonArgs.TutorialHighlightName, ButtonArgs.Command, MultiBox->GetBlocks().Num()));
+	ToolBarButtonBlock->SetCustomMenuDelegate(ButtonArgs.CustomMenuDelegate);
+	ToolBarButtonBlock->SetOnGetMenuContent(ButtonArgs.OnGetMenuContent);
+	ToolBarButtonBlock->SetStyleNameOverride(CurrentStyleOverride);
+
+	MultiBox->AddMultiBlock( ToolBarButtonBlock.ToSharedRef() );
+
+	ApplyHook(ButtonArgs.ExtensionHook, EExtensionHook::After);
+}
diff --git a/Engine/Source/Runtime/Slate/Private/Framework/MultiBox/SClippingVerticalBox.cpp b/Engine/Source/Runtime/Slate/Private/Framework/MultiBox/SClippingVerticalBox.cpp
new file mode 100644
--- /dev/null
+++ b/Engine/Source/Runtime/Slate/Private/Framework/MultiBox/SClippingVerticalBox.cpp
@@ -0,0 +1,151 @@
+// Copyright Epic Games, Inc. All Rights Reserved.
+
+#include "Framework/MultiBox/SClippingVerticalBox.h"
+#include "Layout/ArrangedChildren.h"
+#include "Rendering/DrawElements.h"
+#include "Widgets/Images/SImage.h"
+#include "Widgets/Input/SComboButton.h"
+#include "Framework/Application/SlateApplication.h"
+#include "Framework/MultiBox/SToolBarButtonBlock.h"
+#include "Styling/ToolBarStyle.h"
+#include "Widgets/Images/SLayeredImage.h"
+
+void SClippingVerticalBox::OnArrangeChildren( const FGeometry& AllottedGeometry, FArrangedChildren& ArrangedChildren ) const
+{
+	// If WrapButton hasn't been initialized, that means AddWrapButton() hasn't 
+	// been called and this method isn't going to behave properly
+	check(WrapButton.IsValid());
+
+	LastClippedIdx = ClippedIdx;
+
+	NumClippedChildren = 0;
+	SVerticalBox::OnArrangeChildren(AllottedGeometry, ArrangedChildren);
+
+	// Remove children that are clipped by the allotted geometry
+	const int32 NumChildren = ArrangedChildren.Num();
+	const int32 OverflowButtonIndex = NumChildren - 1;
+	const int32 LastToolBarButtonIndex = OverflowButtonIndex - 1;
+	constexpr int32 OverflowButtonSize = 22;
+	
+	for (int32 ChildIdx = LastToolBarButtonIndex; ChildIdx >= 0; --ChildIdx)
+	{
+		const FArrangedWidget& CurrentWidget = ArrangedChildren[ChildIdx];
+		const FVector2D CurrentWidgetLocalPosition = CurrentWidget.Geometry.GetLocalPositionAtCoordinates(FVector2D::ZeroVector);
+
+		/* if we're not on the last toolbar button, the button furthest Y should also take into account
+		 the height of the overflow button so as not to be positioned over it */
+		const int32 CurrentWidgetMaxY = LastToolBarButtonIndex != ChildIdx ?
+			CurrentWidgetLocalPosition.Y + CurrentWidget.Geometry.Size.Y + OverflowButtonSize :
+			CurrentWidgetLocalPosition.Y + CurrentWidget.Geometry.Size.Y; 
+		
+		if ( CurrentWidgetMaxY > AllottedGeometry.Size.Y)
+		{
+			++NumClippedChildren;
+			ArrangedChildren.Remove(ChildIdx);
+		}
+		else if (LastToolBarButtonIndex == ChildIdx)
+		{
+			ArrangedChildren.Remove(OverflowButtonIndex);	
+			return;
+		}
+	}
+	FArrangedWidget& ArrangedButton = ArrangedChildren[ArrangedChildren.Num() - 1];
+	FVector2D Size = ArrangedButton.Geometry.GetLocalSize();
+	Size.Y = OverflowButtonSize;
+	ArrangedButton.Geometry = AllottedGeometry.MakeChild(Size, FSlateLayoutTransform(AllottedGeometry.GetLocalSize() - Size));
+
+	ClippedIdx = ArrangedChildren.Num() - 1;
+}
+
+int32 SClippingVerticalBox::OnPaint( const FPaintArgs& Args, const FGeometry& AllottedGeometry, const FSlateRect& MyCullingRect, FSlateWindowElementList& OutDrawElements, int32 LayerId, const FWidgetStyle& InWidgetStyle, bool bParentEnabled ) const
+{
+	return SVerticalBox::OnPaint(Args, AllottedGeometry, MyCullingRect, OutDrawElements, LayerId, InWidgetStyle, bParentEnabled);
+}
+
+FVector2D SClippingVerticalBox::ComputeDesiredSize(float LayoutScaleMultiplier) const
+{
+	FVector2D Size = SBoxPanel::ComputeDesiredSize(LayoutScaleMultiplier);
+	{
+		// If the wrap button isn't being shown, subtract it's size from the total desired size
+		const SBoxPanel::FSlot& Child = Children[Children.Num() - 1];
+		const FVector2D& ChildDesiredSize = Child.GetWidget()->GetDesiredSize();
+		Size.Y -= ChildDesiredSize.Y;
+	}
+	return Size;
+}
+
+void SClippingVerticalBox::Construct( const FArguments& InArgs )
+{
+	OnWrapButtonClicked = InArgs._OnWrapButtonClicked;
+	StyleSet = InArgs._StyleSet;
+	StyleName = InArgs._StyleName;
+	bIsFocusable = InArgs._IsFocusable;
+
+	LastClippedIdx = ClippedIdx = INDEX_NONE;
+}
+
+void SClippingVerticalBox::AddWrapButton()
+{
+	const FToolBarStyle& ToolBarStyle = StyleSet->GetWidgetStyle<FToolBarStyle>(StyleName);
+	const TSharedRef<SImage> IconWidget =
+		SNew( SImage )
+		.Visibility(EVisibility::HitTestInvisible)
+		.Image( &ToolBarStyle.ExpandBrush );
+
+	
+	// Construct the wrap button used in toolbars and menubars
+	// Always allow this to be focusable to prevent the menu from collapsing during interaction
+	WrapButton =
+		SNew( SComboButton )
+		.HasDownArrow( false )
+		.ButtonStyle(&ToolBarStyle.ButtonStyle)
+		.ContentPadding( FMargin(2.f, 4.f) )
+		.ToolTipText( NSLOCTEXT("Slate", "ExpandToolbar", "Click to expand toolbar") )
+		.OnGetMenuContent( OnWrapButtonClicked )
+		.Cursor( EMouseCursor::Default )
+		.OnMenuOpenChanged(this, &SClippingVerticalBox::OnWrapButtonOpenChanged)
+		.IsFocusable(true)
+		.ButtonContent()
+		[
+		SNew(SVerticalBox)
+				+ SVerticalBox::Slot()
+				.AutoHeight()
+				.HAlign(HAlign_Center)	// Center the icon horizontally, so that large labels don't stretch out the artwork
+				[
+					IconWidget
+				]
+			
+		];
+
+	AddSlot()
+	.Padding( 0.f )
+	[
+		WrapButton.ToSharedRef()
+	];
+}
+
+void SClippingVerticalBox::OnWrapButtonOpenChanged(bool bIsOpen)
+{
+	if (bIsOpen && !WrapButtonOpenTimer.IsValid())
+	{
+		WrapButtonOpenTimer = RegisterActiveTimer(0.0f, FWidgetActiveTimerDelegate::CreateSP(this, &SClippingVerticalBox::UpdateWrapButtonStatus));
+	}
+	else if(!bIsOpen && WrapButtonOpenTimer.IsValid())
+	{
+		UnRegisterActiveTimer(WrapButtonOpenTimer.ToSharedRef());
+		WrapButtonOpenTimer.Reset();
+	}
+}
+
+EActiveTimerReturnType SClippingVerticalBox::UpdateWrapButtonStatus(double CurrentTime, float DeltaTime)
+{
+	if (LastClippedIdx != ClippedIdx || !WrapButton->IsOpen())
+	{
+		WrapButton->SetIsOpen(false);
+		WrapButtonOpenTimer.Reset();
+		return EActiveTimerReturnType::Stop;
+	}
+
+	return EActiveTimerReturnType::Continue;
+}
+
diff --git a/Engine/Source/Runtime/Slate/Private/Framework/MultiBox/SClippingVerticalBox.h b/Engine/Source/Runtime/Slate/Private/Framework/MultiBox/SClippingVerticalBox.h
new file mode 100644
--- /dev/null
+++ b/Engine/Source/Runtime/Slate/Private/Framework/MultiBox/SClippingVerticalBox.h
@@ -0,0 +1,75 @@
+// Copyright Epic Games, Inc. All Rights Reserved.
+
+#pragma once
+
+#include "CoreMinimal.h"
+#include "Widgets/DeclarativeSyntaxSupport.h"
+#include "Styling/CoreStyle.h"
+#include "Framework/SlateDelegates.h"
+#include "Widgets/SBoxPanel.h"
+
+class FArrangedChildren;
+class FPaintArgs;
+class FSlateWindowElementList;
+class SComboButton;
+
+/** Specialized control for handling the clipping of toolbars and menubars */
+class SClippingVerticalBox : public SVerticalBox
+{
+public:
+	SLATE_BEGIN_ARGS(SClippingVerticalBox) 
+		: _StyleSet(&FCoreStyle::Get())
+		, _StyleName(NAME_None)
+		, _IsFocusable(true)
+		{ }
+
+		SLATE_ARGUMENT(FOnGetContent, OnWrapButtonClicked)
+		SLATE_ARGUMENT(const ISlateStyle*, StyleSet)
+		SLATE_ARGUMENT(FName, StyleName)
+		SLATE_ARGUMENT(bool, IsFocusable)
+	SLATE_END_ARGS()
+
+	/** SWidget interface */
+	virtual int32 OnPaint( const FPaintArgs& Args, const FGeometry& AllottedGeometry, const FSlateRect& MyCullingRect, FSlateWindowElementList& OutDrawElements, int32 LayerId, const FWidgetStyle& InWidgetStyle, bool bParentEnabled ) const override;
+	virtual FVector2D ComputeDesiredSize(float) const override;
+
+	/** SPanel interface */
+	virtual void OnArrangeChildren( const FGeometry& AllottedGeometry, FArrangedChildren& ArrangedChildren ) const override;
+
+	/** Construct this widget */
+	void Construct( const FArguments& InArgs );
+
+	/** Adds the wrap button */
+	void AddWrapButton();
+
+	/** Returns to index of the first clipped child/block */
+	int32 GetClippedIndex() { return ClippedIdx; }
+
+private:
+	void OnWrapButtonOpenChanged(bool bIsOpen);
+	EActiveTimerReturnType UpdateWrapButtonStatus(double CurrentTime, float DeltaTime);
+private:
+	/** The button that is displayed when a toolbar or menubar is clipped */
+	TSharedPtr<SComboButton> WrapButton;
+
+	/** Callback for when the wrap button is clicked */
+	FOnGetContent OnWrapButtonClicked;
+
+	/** Index of the first clipped child/block */
+	mutable int32 ClippedIdx;
+
+	mutable int32 LastClippedIdx;
+
+	/** Number of clipped children not including the wrap button */
+	mutable int32 NumClippedChildren;
+
+	TSharedPtr<FActiveTimerHandle> WrapButtonOpenTimer;
+
+	/** Can the wrap button be focused? */
+	bool bIsFocusable;
+
+	/** The style to use */
+	const ISlateStyle* StyleSet;
+
+	FName StyleName;
+};
diff --git a/Engine/Source/Runtime/Slate/Private/Framework/MultiBox/SToolBarButtonBlock.cpp b/Engine/Source/Runtime/Slate/Private/Framework/MultiBox/SToolBarButtonBlock.cpp
--- a/Engine/Source/Runtime/Slate/Private/Framework/MultiBox/SToolBarButtonBlock.cpp
+++ b/Engine/Source/Runtime/Slate/Private/Framework/MultiBox/SToolBarButtonBlock.cpp
@@ -13,6 +13,19 @@
 #include "Widgets/Images/SLayeredImage.h"
 #include "Widgets/Layout/SSeparator.h"
 
+FToolBarButtonBlock::FToolBarButtonBlock( FButtonArgs ButtonArgs )
+: FMultiBlock( ButtonArgs.Command, ButtonArgs.CommandList, NAME_None, EMultiBlockType::ToolBarButton )
+	, LabelOverride( ButtonArgs.LabelOverride )
+	, ToolTipOverride( ButtonArgs.ToolTipOverride )
+	, IconOverride( ButtonArgs.IconOverride )
+   // , BorderBrushName(ButtonArgs.BorderBrushName)
+	, LabelVisibility()
+	, UserInterfaceActionType(ButtonArgs.UserInterfaceActionType != EUserInterfaceActionType::None ?
+		ButtonArgs.UserInterfaceActionType : EUserInterfaceActionType::Button)
+	, bIsFocusable(false)
+	, bForceSmallIcons(false)
+{
+}
 
 FToolBarButtonBlock::FToolBarButtonBlock( const TSharedPtr< const FUICommandInfo > InCommand, TSharedPtr< const FUICommandList > InCommandList, const TAttribute<FText>& InLabelOverride, const TAttribute<FText>& InToolTipOverride, const TAttribute<FSlateIcon>& InIconOverride )
 	: FMultiBlock( InCommand, InCommandList, NAME_None, EMultiBlockType::ToolBarButton )
@@ -40,7 +53,12 @@ FToolBarButtonBlock::FToolBarButtonBlock( const TAttribute<FText>& InLabel, cons
 
 void FToolBarButtonBlock::SetCustomMenuDelegate( FNewMenuDelegate& InCustomMenuDelegate )
 {
-	CustomMenuDelegate = InCustomMenuDelegate;	
+	CustomMenuDelegate = InCustomMenuDelegate;
+}
+
+void FToolBarButtonBlock::SetOnGetMenuContent( FOnGetContent& InOnGetMenuContent)
+{
+	OnGetMenuContent = InOnGetMenuContent;
 }
 
 void FToolBarButtonBlock::CreateMenuEntry(FMenuBuilder& MenuBuilder) const
@@ -205,7 +223,8 @@ void SToolBarButtonBlock::BuildMultiBlockWidget(const ISlateStyle* StyleSet, con
 
 	// Create the content for our button
 	TSharedRef<SWidget> ButtonContent = SNullWidget::NullWidget;
-	if (MultiBox->GetType() == EMultiBoxType::SlimHorizontalToolBar)
+	if (MultiBox->GetType() == EMultiBoxType::SlimHorizontalToolBar 
+		|| MultiBox->GetType() == EMultiBoxType::SlimHorizontalUniformToolBar)
 	{
 		const FVector2f IconSize = ToolBarStyle.IconSize;
 
@@ -215,6 +234,7 @@ void SToolBarButtonBlock::BuildMultiBlockWidget(const ISlateStyle* StyleSet, con
 			.AddMetaData<FTagMetaData>(FTagMetaData(TutorialHighlightName))
 			+ SHorizontalBox::Slot()
 			.AutoWidth()
+		//    .Padding(ToolBarStyle.IconPadding) @TODO: ~enable with styling changes
 			.VAlign(VAlign_Center)
 			.HAlign(HAlign_Center)
 			[
@@ -308,20 +328,36 @@ void SToolBarButtonBlock::BuildMultiBlockWidget(const ISlateStyle* StyleSet, con
 			CheckStyle = &ToolBarStyle.SettingsToggleButton;
 		}
 
+		const TSharedPtr<SWidget> CheckBox = SNew(SCheckBox)
+						// Use the tool bar style for this check box
+						.Style(CheckStyle)
+						.IsFocusable(bIsFocusable)
+						.ToolTip( FMultiBoxSettings::ToolTipConstructor.Execute( ActualToolTip, nullptr, Action.Pin()))		
+						.OnCheckStateChanged(this, &SToolBarButtonBlock::OnCheckStateChanged )
+						.OnGetMenuContent( ToolBarButtonBlock->OnGetMenuContent )
+						.IsChecked(this, &SToolBarButtonBlock::GetCheckState)
+						.IsEnabled(this, &SToolBarButtonBlock::IsEnabled)
+						[
+							ButtonContent
+						];
+
+		TSharedRef<SWidget> CheckBoxWidget = CheckBox.ToSharedRef();
+
+		if (!ToolBarButtonBlock->BorderBrushName.Get().IsNone())
+		{
+			const FSlateBrush* Brush = FAppStyle::GetBrush(ToolBarButtonBlock->BorderBrushName.Get());
+			CheckBoxWidget =
+				SNew(SBorder)
+				.BorderImage(Brush)
+				.Padding(2.f)
+				[
+					CheckBox.ToSharedRef()
+				];
+		}
+
 		ChildSlot
 		[
-			// Create a check box
-			SNew(SCheckBox)
-			// Use the tool bar style for this check box
-			.Style(CheckStyle)
-			.IsFocusable(bIsFocusable)
-			.ToolTip( FMultiBoxSettings::ToolTipConstructor.Execute( ActualToolTip, nullptr, Action.Pin()))		
-			.OnCheckStateChanged(this, &SToolBarButtonBlock::OnCheckStateChanged )
-			.IsChecked(this, &SToolBarButtonBlock::GetCheckState)
-			.IsEnabled(this, &SToolBarButtonBlock::IsEnabled)
-			[
-				ButtonContent
-			]
+			CheckBoxWidget
 		];
 	}
 
@@ -366,6 +402,10 @@ void SToolBarButtonBlock::BuildMultiBlockWidget(const ISlateStyle* StyleSet, con
 	SetVisibility(TAttribute<EVisibility>(this, &SToolBarButtonBlock::GetBlockVisibility));
 }
 
+bool FToolBarButtonBlock::GetIsFocusable() const
+{
+	return bIsFocusable;
+}
 
 
 /**
diff --git a/Engine/Source/Runtime/Slate/Private/Framework/MultiBox/SToolBarComboButtonBlock.cpp b/Engine/Source/Runtime/Slate/Private/Framework/MultiBox/SToolBarComboButtonBlock.cpp
--- a/Engine/Source/Runtime/Slate/Private/Framework/MultiBox/SToolBarComboButtonBlock.cpp
+++ b/Engine/Source/Runtime/Slate/Private/Framework/MultiBox/SToolBarComboButtonBlock.cpp
@@ -93,7 +93,8 @@ void SToolBarComboButtonBlock::BuildMultiBlockWidget(const ISlateStyle* StyleSet
 
 		ActualIconWidget->AddLayer(TAttribute<const FSlateBrush*>(this, &SToolBarComboButtonBlock::GetOverlayIconBrush));
 
-		if (MultiBox->GetType() == EMultiBoxType::SlimHorizontalToolBar)
+		if (MultiBox->GetType() == EMultiBoxType::SlimHorizontalToolBar
+		|| MultiBox->GetType() == EMultiBoxType::SlimHorizontalUniformToolBar)
 		{
 			const FVector2f IconSize = ToolBarStyle.IconSize;
 
@@ -129,7 +130,8 @@ void SToolBarComboButtonBlock::BuildMultiBlockWidget(const ISlateStyle* StyleSet
 	
 	// Create the content for our button
 	TSharedRef<SWidget> ButtonContent = SNullWidget::NullWidget;
-	if (MultiBox->GetType() == EMultiBoxType::SlimHorizontalToolBar)
+	if (MultiBox->GetType() == EMultiBoxType::SlimHorizontalToolBar
+		|| MultiBox->GetType() == EMultiBoxType::SlimHorizontalUniformToolBar)
 	{
 		ButtonContent =
 			SNew(SHorizontalBox)
diff --git a/Engine/Source/Runtime/Slate/Private/Framework/MultiBox/SToolBarStackButtonBlock.cpp b/Engine/Source/Runtime/Slate/Private/Framework/MultiBox/SToolBarStackButtonBlock.cpp
--- a/Engine/Source/Runtime/Slate/Private/Framework/MultiBox/SToolBarStackButtonBlock.cpp
+++ b/Engine/Source/Runtime/Slate/Private/Framework/MultiBox/SToolBarStackButtonBlock.cpp
@@ -280,7 +280,8 @@ void SToolBarStackButtonBlock::BuildMultiBlockWidget(const ISlateStyle* StyleSet
 	FMargin ExtraLabelPadding = FMargin(4, 0, 0, 0);
 
 	TSharedRef<SWidget> ButtonContent = SNullWidget::NullWidget;
-	if (MultiBox->GetType() == EMultiBoxType::SlimHorizontalToolBar)
+	if (MultiBox->GetType() == EMultiBoxType::SlimHorizontalToolBar
+		|| MultiBox->GetType() == EMultiBoxType::SlimHorizontalUniformToolBar)
 	{
 		const FVector2f IconSize = ToolBarStyle.IconSize;
 
diff --git a/Engine/Source/Runtime/Slate/Private/Framework/MultiBox/SUniformToolbarButtonBlock.cpp b/Engine/Source/Runtime/Slate/Private/Framework/MultiBox/SUniformToolbarButtonBlock.cpp
new file mode 100644
--- /dev/null
+++ b/Engine/Source/Runtime/Slate/Private/Framework/MultiBox/SUniformToolbarButtonBlock.cpp
@@ -0,0 +1,16 @@
+#include "Framework/MultiBox/SUniformToolbarButtonBlock.h"
+
+FUniformToolbarButtonBlock::FUniformToolbarButtonBlock(FButtonArgs ButtonArgs) :
+	FToolBarButtonBlock(ButtonArgs)
+{
+}
+
+TSharedRef< class IMultiBlockBaseWidget > FUniformToolbarButtonBlock::ConstructWidget() const
+{
+	return SNew( SUniformToolbarButtonBlock )
+		.LabelVisibility(EVisibility::Visible)
+		.IsFocusable(GetIsFocusable())
+		.ForceSmallIcons(true)
+		.TutorialHighlightName(GetTutorialHighlightName())
+		.Cursor(EMouseCursor::Default);
+}
\ No newline at end of file
diff --git a/Engine/Source/Runtime/Slate/Private/Widgets/Docking/SDockTab.cpp b/Engine/Source/Runtime/Slate/Private/Widgets/Docking/SDockTab.cpp
--- a/Engine/Source/Runtime/Slate/Private/Widgets/Docking/SDockTab.cpp
+++ b/Engine/Source/Runtime/Slate/Private/Widgets/Docking/SDockTab.cpp
@@ -995,3 +995,12 @@ void SDockTab::UpdateActivationTime()
 		LastActivationTime = FSlateApplication::Get().GetCurrentTime();
 	}
 }
+
+
+void SDockTab::SetParentDockTabStackTabWellHidden(bool bIsTabWellHidden)
+{
+	if ( TSharedPtr<SDockingTabStack> ParentStack = GetParentDockTabStack() )
+	{
+		ParentStack->SetTabWellHidden(bIsTabWellHidden);
+	}
+}
\ No newline at end of file
diff --git a/Engine/Source/Runtime/Slate/Private/Widgets/Layout/SUniformWrapPanel.cpp b/Engine/Source/Runtime/Slate/Private/Widgets/Layout/SUniformWrapPanel.cpp
--- a/Engine/Source/Runtime/Slate/Private/Widgets/Layout/SUniformWrapPanel.cpp
+++ b/Engine/Source/Runtime/Slate/Private/Widgets/Layout/SUniformWrapPanel.cpp
@@ -10,6 +10,7 @@ void SUniformWrapPanel::PrivateRegisterAttributes(FSlateAttributeInitializer& At
 	SLATE_ADD_MEMBER_ATTRIBUTE_DEFINITION(AttributeInitializer, MinDesiredSlotWidth, EInvalidateWidgetReason::Layout);
 	SLATE_ADD_MEMBER_ATTRIBUTE_DEFINITION(AttributeInitializer, MinDesiredSlotHeight, EInvalidateWidgetReason::Layout);
 	SLATE_ADD_MEMBER_ATTRIBUTE_DEFINITION(AttributeInitializer, MaxDesiredSlotWidth, EInvalidateWidgetReason::Layout);
+	SLATE_ADD_MEMBER_ATTRIBUTE_DEFINITION(AttributeInitializer, NumColumnsOverride, EInvalidateWidgetReason::Layout);
 	SLATE_ADD_MEMBER_ATTRIBUTE_DEFINITION(AttributeInitializer, MaxDesiredSlotHeight, EInvalidateWidgetReason::Layout);
 	SLATE_ADD_MEMBER_ATTRIBUTE_DEFINITION(AttributeInitializer, HAlign, EInvalidateWidgetReason::Paint);
 	SLATE_ADD_MEMBER_ATTRIBUTE_DEFINITION(AttributeInitializer, EvenRowDistribution, EInvalidateWidgetReason::Paint);
@@ -18,6 +19,7 @@ void SUniformWrapPanel::PrivateRegisterAttributes(FSlateAttributeInitializer& At
 SUniformWrapPanel::SUniformWrapPanel()
 	: Children(this)
 	, SlotPadding(*this, FMargin(0.0f))
+	, NumColumnsOverride(*this, 0)
 	, MinDesiredSlotWidth(*this, 0.f)
 	, MinDesiredSlotHeight(*this, 0.f)
 	, MaxDesiredSlotWidth(*this, FLT_MAX)
@@ -36,6 +38,7 @@ void SUniformWrapPanel::Construct( const FArguments& InArgs )
 	MinDesiredSlotWidth.Assign(*this, InArgs._MinDesiredSlotWidth);
 	MinDesiredSlotHeight.Assign(*this, InArgs._MinDesiredSlotHeight);
 	MaxDesiredSlotWidth.Assign(*this, InArgs._MaxDesiredSlotWidth);
+	NumColumnsOverride.Assign(*this, InArgs._NumColumnsOverride);
 	MaxDesiredSlotHeight.Assign(*this, InArgs._MaxDesiredSlotHeight);
 	HAlign.Assign(*this, InArgs._HAlign);
 	EvenRowDistribution.Assign(*this, InArgs._EvenRowDistribution);
@@ -48,7 +51,9 @@ void SUniformWrapPanel::OnArrangeChildren( const FGeometry& AllottedGeometry, FA
 	if ( Children.Num() > 0)
 	{
 		FVector2D CellSize = ComputeUniformCellSize();
-		NumColumns = FMath::Max(1, FMath::Min(NumVisibleChildren, FMath::FloorToInt( AllottedGeometry.GetLocalSize().X / CellSize.X )));
+		NumColumns = NumColumnsOverride.Get() != 0 ?
+				NumColumnsOverride.Get() :
+		FMath::Max(1, FMath::Min(NumVisibleChildren, FMath::FloorToInt( AllottedGeometry.GetLocalSize().X / CellSize.X )));
 		NumRows = FMath::CeilToInt ( (float) NumVisibleChildren / (float) NumColumns );
 
 		// If we have to have N rows, try to distribute the items across the rows evenly
@@ -172,11 +177,13 @@ FVector2D SUniformWrapPanel::ComputeDesiredSize( float ) const
 		const FVector2D& LocalSize = GetTickSpaceGeometry().GetLocalSize();
 		if (!LocalSize.IsZero()) 
 		{
-			NumColumns = FMath::FloorToInt(LocalSize.X / MaxChildDesiredSize.X);
+			NumColumns = NumColumnsOverride.Get() != 0 ? NumColumnsOverride.Get() : 
+				FMath::FloorToInt(LocalSize.X / MaxChildDesiredSize.X);
 		}
 		else 
 		{
-			NumColumns = FMath::CeilToInt(FMath::Sqrt((float)NumVisibleChildren));
+			NumColumns = NumColumnsOverride.Get() != 0 ? NumColumnsOverride.Get() :
+			    FMath::CeilToInt(FMath::Sqrt((float)NumVisibleChildren));
 		}
 
 		if (NumColumns > 0)
@@ -214,6 +221,11 @@ void SUniformWrapPanel::SetMaxDesiredSlotWidth(TAttribute<float> InMaxDesiredSlo
 	MaxDesiredSlotWidth.Assign(*this, MoveTemp(InMaxDesiredSlotWidth));
 }
 
+void SUniformWrapPanel::SetNumColumnsOverride(TAttribute<uint32> InNumColumnsOverride)
+{
+	NumColumnsOverride.Assign(*this, MoveTemp(InNumColumnsOverride));
+}
+
 void SUniformWrapPanel::SetMaxDesiredSlotHeight(TAttribute<float> InMaxDesiredSlotHeight)
 {
 	MaxDesiredSlotHeight.Assign(*this, MoveTemp(InMaxDesiredSlotHeight));
diff --git a/Engine/Source/Runtime/Slate/Public/Framework/MultiBox/MultiBox.h b/Engine/Source/Runtime/Slate/Public/Framework/MultiBox/MultiBox.h
--- a/Engine/Source/Runtime/Slate/Public/Framework/MultiBox/MultiBox.h
+++ b/Engine/Source/Runtime/Slate/Public/Framework/MultiBox/MultiBox.h
@@ -896,6 +896,9 @@ private:
 
 	/** Specialized box widget to handle clipping of toolbars and menubars */
 	TSharedPtr<class SClippingHorizontalBox> ClippedHorizontalBox;
+	
+	/** Specialized box widget to handle clipping of vertical toolbars  */
+	TSharedPtr<class SClippingVerticalBox> ClippedVerticalBox;
 
 	/** Specialized box widget to handle clipping of toolbars and menubars */
 	TSharedPtr<SUniformWrapPanel> UniformToolbarPanel;
diff --git a/Engine/Source/Runtime/Slate/Public/Framework/MultiBox/MultiBoxBuilder.h b/Engine/Source/Runtime/Slate/Public/Framework/MultiBox/MultiBoxBuilder.h
--- a/Engine/Source/Runtime/Slate/Public/Framework/MultiBox/MultiBoxBuilder.h
+++ b/Engine/Source/Runtime/Slate/Public/Framework/MultiBox/MultiBoxBuilder.h
@@ -5,6 +5,7 @@
 #include "Framework/MultiBox/MultiBoxExtender.h"
 #include "Framework/SlateDelegates.h"
 #include "CoreMinimal.h"
+#include "SToolBarButtonBlock.h"
 #include "Misc/Attribute.h"
 #include "Layout/Visibility.h"
 #include "Widgets/SWidget.h"
@@ -21,6 +22,7 @@ class FUICommandInfo;
 class FUICommandList;
 struct FSlateIcon;
 struct FUIAction;
+struct FButtonArgs;
 
 /** Delegate used by multi-box to call a user function to populate a new menu.  Used for spawning sub-menus and pull-down menus. */
 DECLARE_DELEGATE_OneParam( FNewMenuDelegate, class FMenuBuilder& );
@@ -512,6 +514,14 @@ public:
 
 	void SetIsFocusable(bool bInIsFocusable);
 
+	/**
+	 * Adds a tool bar button
+	 *
+	 * @param	ButtonArgs The Parameters object which will provide the data to initialize the button
+	 */
+	virtual void AddToolBarButton(FButtonArgs& ButtonArgs);
+
+
 	/**
 	 * Adds a tool bar button
 	 *
@@ -627,6 +637,8 @@ protected:
 	virtual void ApplyHook(FName InExtensionHook, EExtensionHook::Position HookPosition) override;
 	virtual void ApplySectionBeginning() override;
 
+	void InitializeToolBarButtonBlock(TSharedPtr<FToolBarButtonBlock> ButtonRowBlock, FButtonArgs& ButtonArgs);
+
 private:
 	/** Current extension hook name for sections to determine where sections begin and end */
 	FName CurrentSectionExtensionHook;
@@ -657,6 +669,7 @@ public:
 	FVerticalToolBarBuilder(TSharedPtr<const FUICommandList> InCommandList, FMultiBoxCustomization InCustomization, TSharedPtr<FExtender> InExtender = nullptr, const bool InForceSmallIcons = false)
 		: FToolBarBuilder(EMultiBoxType::VerticalToolBar, InCommandList, InCustomization, InExtender, InForceSmallIcons)
 	{
+		this->SetStyle(&FAppStyle::Get(), "FVerticalToolBar");
 	}
 };
 
@@ -730,3 +743,17 @@ protected:
 	/** FMultiBoxBuilder interface */
 	virtual void ApplyHook(FName InExtensionHook, EExtensionHook::Position HookPosition) override {}
 };
+
+class SLATE_API FSlimHorizontalUniformToolBarBuilder : public FToolBarBuilder
+{
+public:
+	/**
+	 * Constructor
+	 *
+	 * @param	InCommandList	The action list that maps command infos to delegates that should be called for each command associated with a multiblock widget
+	 */
+	FSlimHorizontalUniformToolBarBuilder(TSharedPtr<const FUICommandList> InCommandList, FMultiBoxCustomization InCustomization, TSharedPtr<FExtender> InExtender = nullptr, const bool InForceSmallIcons = false);
+
+	void AddToolBarButton( FButtonArgs& ButtonArgs ) override;
+
+};
\ No newline at end of file
diff --git a/Engine/Source/Runtime/Slate/Public/Framework/MultiBox/MultiBoxDefs.h b/Engine/Source/Runtime/Slate/Public/Framework/MultiBox/MultiBoxDefs.h
--- a/Engine/Source/Runtime/Slate/Public/Framework/MultiBox/MultiBoxDefs.h
+++ b/Engine/Source/Runtime/Slate/Public/Framework/MultiBox/MultiBoxDefs.h
@@ -37,6 +37,9 @@ enum class EMultiBoxType : uint8
 
 	/** Buttons arranged in rows, with a maximum number of buttons per row, like a toolbar but can have multiple rows*/
 	ButtonRow,
+
+	/** A toolbar with horizontally-oriented buttons that tries to arrange all toolbar items uniformly */
+	SlimHorizontalUniformToolBar,
 };
 
 
diff --git a/Engine/Source/Runtime/Slate/Public/Framework/MultiBox/SToolBarButtonBlock.h b/Engine/Source/Runtime/Slate/Public/Framework/MultiBox/SToolBarButtonBlock.h
--- a/Engine/Source/Runtime/Slate/Public/Framework/MultiBox/SToolBarButtonBlock.h
+++ b/Engine/Source/Runtime/Slate/Public/Framework/MultiBox/SToolBarButtonBlock.h
@@ -15,6 +15,29 @@
 
 class SBorder;
 
+/** Delegate used by multi-box to call a user function to populate a new menu.  Used for spawning sub-menus and pull-down menus. */
+DECLARE_DELEGATE_OneParam( FNewMenuDelegate, class FMenuBuilder& );
+
+struct SLATE_API FButtonArgs : public TSharedFromThis<FButtonArgs>
+{
+	TSharedPtr< const FUICommandInfo > Command;
+	TSharedPtr< const FUICommandList > CommandList;
+	FName ExtensionHook;
+	TAttribute<FText> LabelOverride;
+	TAttribute<FText> ToolTipOverride;
+	TAttribute<FSlateIcon> IconOverride;
+	FName TutorialHighlightName;
+	FName BorderBrushName;
+	EUserInterfaceActionType UserInterfaceActionType = EUserInterfaceActionType::None;
+	FUIAction Action;
+	
+	FNewMenuDelegate CustomMenuDelegate;
+	FOnGetContent OnGetMenuContent;
+
+	explicit FButtonArgs() {}
+};
+
+
 /**
  * Tool bar button MultiBlock
  */
@@ -24,6 +47,9 @@ class SLATE_API FToolBarButtonBlock
 
 public:
 
+	
+	FToolBarButtonBlock( FButtonArgs ButtonArgs );
+	
 	/**
 	 * Constructor
 	 *
@@ -53,6 +79,9 @@ public:
 	/** Set whether this toolbar should always use small icons, regardless of the current settings */
 	void SetForceSmallIcons( const bool InForceSmallIcons ) { bForceSmallIcons = InForceSmallIcons; }
 
+	/** Set the border brush for this toolbar button block */
+	void SetBorderBrushName( const FName InBorderBrushName ) { BorderBrushName = InBorderBrushName; }
+
 	/** FMultiBlock interface */
 	virtual void CreateMenuEntry(class FMenuBuilder& MenuBuilder) const override;
 	virtual bool HasIcon() const override;
@@ -60,6 +89,12 @@ public:
 	/** Set optional delegate to customize when a menu appears instead of the widget, such as in toolbars */
 	void SetCustomMenuDelegate( FNewMenuDelegate& InOnFillMenuDelegate);
 
+	void SetOnGetMenuContent( FOnGetContent& OnGetMenuContent);
+
+protected:
+	
+	bool GetIsFocusable() const;
+
 private:
 
 	/** FMultiBlock private interface */
@@ -75,6 +110,9 @@ private:
 
 	/** Optional overridden tool tip for this tool bar button.  If not set, then the action's tool tip will be used instead. */
 	TAttribute<FText> ToolTipOverride;
+	
+	/** The name of the border brush style, if present */
+	TAttribute<FName> BorderBrushName;
 
 	/** Optional overridden icon for this tool bar button.  IF not set, then the action's icon will be used instead. */
 	TAttribute<FSlateIcon> IconOverride;
@@ -93,6 +131,9 @@ private:
 
 	/** Optional delegate to customize when a menu appears instead of the widget, such as in toolbars */
 	FNewMenuDelegate CustomMenuDelegate;
+
+	/** Delegate to execute to get the menu content of this button */
+	FOnGetContent OnGetMenuContent;
 };
 
 
@@ -138,8 +179,7 @@ protected:
 	 * Called by Slate when this tool bar button's button is clicked
 	 */
 	FReply OnClicked();
-
-
+	
 	/**
 	 * Called by Slate when this tool bar check box button is toggled
 	 */
@@ -206,4 +246,5 @@ private:
 
 	/** Name to identify a widget for tutorials */
 	FName TutorialHighlightName;
+
 };
diff --git a/Engine/Source/Runtime/Slate/Public/Framework/MultiBox/SUniformToolbarButtonBlock.h b/Engine/Source/Runtime/Slate/Public/Framework/MultiBox/SUniformToolbarButtonBlock.h
new file mode 100644
--- /dev/null
+++ b/Engine/Source/Runtime/Slate/Public/Framework/MultiBox/SUniformToolbarButtonBlock.h
@@ -0,0 +1,39 @@
+// Copyright Epic Games, Inc. All Rights Reserved.
+
+#pragma once
+
+#include "Framework/MultiBox/SToolBarButtonBlock.h"
+
+struct FButtonArgs;
+/**
+ * Horizontal Button Row MultiBlock
+ */
+class SLATE_API FUniformToolbarButtonBlock
+	: public FToolBarButtonBlock
+{
+
+public:
+
+	/**
+	 * Constructor
+	 *
+	 * @param	ButtonArgs The Parameters object which will provide the data to initialize the button
+	 */
+	FUniformToolbarButtonBlock( FButtonArgs ButtonArgs);
+
+	
+	/**
+	 * Allocates a widget for this type of MultiBlock.  Override this in derived classes.
+	 *
+	 * @return  MultiBlock widget object
+	 */
+	 virtual TSharedRef< class IMultiBlockBaseWidget > ConstructWidget() const override;
+
+};
+
+class SLATE_API SUniformToolbarButtonBlock
+	: public SToolBarButtonBlock
+{
+	
+};
+
diff --git a/Engine/Source/Runtime/Slate/Public/SlateFwd.h b/Engine/Source/Runtime/Slate/Public/SlateFwd.h
--- a/Engine/Source/Runtime/Slate/Public/SlateFwd.h
+++ b/Engine/Source/Runtime/Slate/Public/SlateFwd.h
@@ -38,6 +38,7 @@ class SInlineEditableTextBlock;
 class SVirtualKeyboardEntry;
 class SSafeZone;
 struct FMarqueeRect;
+struct FButtonArgs;
 
 template <typename NumericType, typename VectorType, int32 NumberOfComponent> class SNumericVectorInputBox;
 using SVectorInputBox = SNumericVectorInputBox<float, UE::Math::TVector<float>, 3>;
diff --git a/Engine/Source/Runtime/Slate/Public/Widgets/Docking/SDockTab.h b/Engine/Source/Runtime/Slate/Public/Widgets/Docking/SDockTab.h
--- a/Engine/Source/Runtime/Slate/Public/Widgets/Docking/SDockTab.h
+++ b/Engine/Source/Runtime/Slate/Public/Widgets/Docking/SDockTab.h
@@ -335,6 +335,8 @@ public:
 		return LastActivationTime;
 	}
 
+	void SetParentDockTabStackTabWellHidden(bool bIsTabWellHidden);
+
 protected:
 	/** Provide a default tab icon. */
 	void ProvideDefaultIcon(const FSlateBrush* InDefaultIcon);
diff --git a/Engine/Source/Runtime/Slate/Public/Widgets/Layout/SUniformWrapPanel.h b/Engine/Source/Runtime/Slate/Public/Widgets/Layout/SUniformWrapPanel.h
--- a/Engine/Source/Runtime/Slate/Public/Widgets/Layout/SUniformWrapPanel.h
+++ b/Engine/Source/Runtime/Slate/Public/Widgets/Layout/SUniformWrapPanel.h
@@ -50,6 +50,7 @@ public:
 
 	SLATE_BEGIN_ARGS( SUniformWrapPanel )
 		: _SlotPadding( FMargin(0.0f) )
+		, _NumColumnsOverride( 0 )
 		, _MinDesiredSlotWidth( 0.0f )
 		, _MinDesiredSlotHeight( 0.0f )
 		, _MaxDesiredSlotWidth( FLT_MAX )
@@ -66,6 +67,9 @@ public:
 		/** Padding given to each slot */
 		SLATE_ATTRIBUTE(FMargin, SlotPadding)
 
+		/** The number of columns for the wrapped panel */
+		SLATE_ATTRIBUTE(uint32, NumColumnsOverride)
+	
 		/** The minimum desired width of the slots */
 		SLATE_ATTRIBUTE(float, MinDesiredSlotWidth)
 
@@ -102,6 +106,9 @@ public:
 	/** See MinDesiredSlotHeight attribute */
 	void SetMinDesiredSlotHeight(TAttribute<float> InMinDesiredSlotHeight);
 
+	/** See NumColumnsOverride attribute */
+	void SetNumColumnsOverride(TAttribute<uint32> InNumColumnsOverride);
+	
 	/** See MinDesiredSlotWidth attribute */
 	void SetMaxDesiredSlotWidth(TAttribute<float> InMaxDesiredSlotWidth);
 
@@ -146,6 +153,8 @@ private:
 	mutable int32 NumRows;
 	mutable int32 NumVisibleChildren;
 
+	TSlateAttribute<uint32> NumColumnsOverride;
+	
 	TSlateAttribute<float> MinDesiredSlotWidth;
 	TSlateAttribute<float> MinDesiredSlotHeight;
 
