Fix for copying in a loop

diff --git a/Engine/Source/Developer/WidgetRegistration/Private/ToolkitBuilder.cpp b/Engine/Source/Developer/WidgetRegistration/Private/ToolkitBuilder.cpp
--- a/Engine/Source/Developer/WidgetRegistration/Private/ToolkitBuilder.cpp
+++ b/Engine/Source/Developer/WidgetRegistration/Private/ToolkitBuilder.cpp
@@ -99,7 +99,7 @@ void FToolkitBuilder::UpdateEditablePalette(FEditablePalette& Palette)
 
 void FToolkitBuilder::UpdateWidget()
 {
-	for (const TSharedRef<FEditablePalette> EditablePalette : EditablePalettesArray)
+	for (const TSharedRef<FEditablePalette>& EditablePalette : EditablePalettesArray)
 	{
 		UpdateEditablePalette(*EditablePalette);		
 	}
@@ -225,7 +225,7 @@ TSharedRef<SWidget> FToolkitBuilder::GetContextMenuContent(const FName CommandNa
 {
 	FMenuBuilder MenuBuilder(true, nullptr);
 
-	for (const TSharedRef<FEditablePalette> EditablePalette : EditablePalettesArray)
+	for (const TSharedRef<FEditablePalette>& EditablePalette : EditablePalettesArray)
 	{
 		const FUIAction ItemAction(FExecuteAction::CreateSP(
 			SharedThis(this), 
