Updates for Modular Class Display to add a widget registration system and get it working for BaseToolkit and toolbars.

diff --git a/Engine/Source/Developer/WidgetRegistration/Private/ToolElementRegistry.cpp b/Engine/Source/Developer/WidgetRegistration/Private/ToolElementRegistry.cpp
--- a/Engine/Source/Developer/WidgetRegistration/Private/ToolElementRegistry.cpp
+++ b/Engine/Source/Developer/WidgetRegistration/Private/ToolElementRegistry.cpp
@@ -1,7 +1,5 @@
 // Copyright Epic Games, Inc. All Rights Reserved.
 
-#pragma once
-
 #include "ToolElementRegistry.h"
 
 #include "Widgets/SNullWidget.h"
diff --git a/Engine/Source/Developer/WidgetRegistration/Private/ToolbarRegistrationArgs.cpp b/Engine/Source/Developer/WidgetRegistration/Private/ToolbarRegistrationArgs.cpp
--- a/Engine/Source/Developer/WidgetRegistration/Private/ToolbarRegistrationArgs.cpp
+++ b/Engine/Source/Developer/WidgetRegistration/Private/ToolbarRegistrationArgs.cpp
@@ -1,6 +1,5 @@
 // Copyright Epic Games, Inc. All Rights Reserved.
 
-#pragma once
 #include "ToolbarRegistrationArgs.h"
 
 
