Fix styling inclusion build issue

diff --git a/Engine/Source/Developer/WidgetRegistration/Public/FToolkitWidgetStyle.h b/Engine/Source/Developer/WidgetRegistration/Public/FToolkitWidgetStyle.h
--- a/Engine/Source/Developer/WidgetRegistration/Public/FToolkitWidgetStyle.h
+++ b/Engine/Source/Developer/WidgetRegistration/Public/FToolkitWidgetStyle.h
@@ -7,6 +7,7 @@
 #include "Brushes/SlateNoResource.h"
 
 #include "FToolkitWidgetStyle.generated.h"
+#include "Styling/SlateWidgetStyle.h"
 
 /**
  * FToolkitWidgetStyle is the FSlateWidgetStyle that defines the styling of a ToolkitWidget 
