Update font style include

diff --git a/Engine/Source/Developer/WidgetRegistration/Public/FToolkitWidgetStyle.h b/Engine/Source/Developer/WidgetRegistration/Public/FToolkitWidgetStyle.h
--- a/Engine/Source/Developer/WidgetRegistration/Public/FToolkitWidgetStyle.h
+++ b/Engine/Source/Developer/WidgetRegistration/Public/FToolkitWidgetStyle.h
@@ -7,6 +7,7 @@
 #include "Brushes/SlateNoResource.h"
 
 #include "Styling/SlateWidgetStyle.h"
+#include "Fonts/SlateFontInfo.h"
 #include "FToolkitWidgetStyle.generated.h"
 
 /**
