fix include build issue

Update font style include

diff --git a/Engine/Source/Developer/WidgetRegistration/Private/ToolkitStyle.cpp b/Engine/Source/Developer/WidgetRegistration/Private/ToolkitStyle.cpp
--- a/Engine/Source/Developer/WidgetRegistration/Private/ToolkitStyle.cpp
+++ b/Engine/Source/Developer/WidgetRegistration/Private/ToolkitStyle.cpp
@@ -8,6 +8,7 @@
 #include "Styling/SlateStyle.h"
 #include "Styling/StarshipCoreStyle.h"
 #include "Styling/StyleColors.h"
+#include "Styling/CoreStyle.h"
 
 TSharedPtr< FSlateStyleSet > FToolkitStyle::StyleSet = nullptr;
 FName FToolkitStyle::StyleName("ToolkitStyle");
