Fixes log errors for undefined styles.

diff --git a/Engine/Source/Developer/WidgetRegistration/Private/ToolkitStyle.cpp b/Engine/Source/Developer/WidgetRegistration/Private/ToolkitStyle.cpp
--- a/Engine/Source/Developer/WidgetRegistration/Private/ToolkitStyle.cpp
+++ b/Engine/Source/Developer/WidgetRegistration/Private/ToolkitStyle.cpp
@@ -26,9 +26,8 @@ void FToolkitStyle::Initialize()
 		StyleSet = Create();
 		FSlateStyleRegistry::RegisterSlateStyle(*StyleSet);
 	}
-
 	
-	FToolkitWidgetStyle Style = FStarshipCoreStyle::GetCoreStyle().GetWidgetStyle<FToolkitWidgetStyle>("FToolkitWidget");
+	FToolkitWidgetStyle Style;
 
 	Style.SetTitleBackgroundBrush(FSlateRoundedBoxBrush(FStyleColors::Header, 0.f, FLinearColor(0, 0, 0, .8), 0.5))
 		 .SetToolDetailsBackgroundBrush( FSlateRoundedBoxBrush(FStyleColors::Background, 0.f, FLinearColor(0, 0, 0, .8), 0.5))
