WidgetRegistration: Add preliminary framework for saving/loading favorites using EditorConfig

diff --git a/Engine/Plugins/Editor/ModelingToolsEditorMode/Source/ModelingToolsEditorMode/ModelingToolsEditorMode.Build.cs b/Engine/Plugins/Editor/ModelingToolsEditorMode/Source/ModelingToolsEditorMode/ModelingToolsEditorMode.Build.cs
--- a/Engine/Plugins/Editor/ModelingToolsEditorMode/Source/ModelingToolsEditorMode/ModelingToolsEditorMode.Build.cs
+++ b/Engine/Plugins/Editor/ModelingToolsEditorMode/Source/ModelingToolsEditorMode/ModelingToolsEditorMode.Build.cs
@@ -51,6 +51,7 @@ public class ModelingToolsEditorMode : ModuleRules
 				"DeveloperSettings",
 				"PropertyEditor",
 				"ToolMenus",
+				"EditorConfig",
 				"PresetAsset"
 				// ... add private dependencies that you statically link with here ...	
 			}
diff --git a/Engine/Plugins/Editor/ModelingToolsEditorMode/Source/ModelingToolsEditorMode/Private/ModelingToolsEditablePaletteConfig.cpp b/Engine/Plugins/Editor/ModelingToolsEditorMode/Source/ModelingToolsEditorMode/Private/ModelingToolsEditablePaletteConfig.cpp
new file mode 100644
--- /dev/null
+++ b/Engine/Plugins/Editor/ModelingToolsEditorMode/Source/ModelingToolsEditorMode/Private/ModelingToolsEditablePaletteConfig.cpp
@@ -0,0 +1,39 @@
+// Copyright Epic Games, Inc. All Rights Reserved.
+
+#include "ModelingToolsEditablePaletteConfig.h"
+
+TObjectPtr<UModelingModeEditableToolPaletteConfig> UModelingModeEditableToolPaletteConfig::Instance = nullptr;
+
+void UModelingModeEditableToolPaletteConfig::Initialize()
+{
+	if(!Instance)
+	{
+		Instance = NewObject<UModelingModeEditableToolPaletteConfig>(); 
+		Instance->AddToRoot();
+	}
+}
+
+FEditableToolPaletteSettings* UModelingModeEditableToolPaletteConfig::GetMutablePaletteConfig(const FName& InstanceName)
+{
+	if (InstanceName.IsNone())
+	{
+		return nullptr;
+	}
+
+	return &EditableToolPalettes.FindOrAdd(InstanceName);
+}
+
+const FEditableToolPaletteSettings* UModelingModeEditableToolPaletteConfig::GetConstPaletteConfig(const FName& InstanceName)
+{
+	if (InstanceName.IsNone())
+	{
+		return nullptr;
+	}
+
+	return EditableToolPalettes.Find(InstanceName);
+}
+
+void UModelingModeEditableToolPaletteConfig::SavePaletteConfig(const FName&)
+{
+	SaveEditorConfig();
+}
\ No newline at end of file
diff --git a/Engine/Plugins/Editor/ModelingToolsEditorMode/Source/ModelingToolsEditorMode/Private/ModelingToolsEditablePaletteConfig.h b/Engine/Plugins/Editor/ModelingToolsEditorMode/Source/ModelingToolsEditorMode/Private/ModelingToolsEditablePaletteConfig.h
new file mode 100644
--- /dev/null
+++ b/Engine/Plugins/Editor/ModelingToolsEditorMode/Source/ModelingToolsEditorMode/Private/ModelingToolsEditablePaletteConfig.h
@@ -0,0 +1,37 @@
+// Copyright Epic Games, Inc. All Rights Reserved.
+
+#pragma once
+
+#include "Containers/Map.h"
+#include "EditorConfigBase.h"
+#include "ToolkitBuilderConfig.h"
+#include "UObject/ObjectPtr.h"
+
+#include "ModelingToolsEditablePaletteConfig.generated.h"
+
+/* Implementation of IEditableToolPaletteConfigManager specific to ModelingMode, currently needed because we cannot have
+ * one at the FModeToolkit level due to EditorConfig depending on UnrealEd which is where the mode toolkit lives
+ */
+UCLASS(EditorConfig="EditableToolPalette")
+class UModelingModeEditableToolPaletteConfig : public UEditorConfigBase, public IEditableToolPaletteConfigManager
+{
+	GENERATED_BODY()
+	
+public:
+
+	// IEditableToolPaletteConfigManager Interface
+	virtual FEditableToolPaletteSettings* GetMutablePaletteConfig(const FName& InstanceName) override;
+	virtual const FEditableToolPaletteSettings* GetConstPaletteConfig(const FName& InstanceName) override;
+	virtual void SavePaletteConfig(const FName& InstanceName) override;
+
+	static void Initialize();
+	static UModelingModeEditableToolPaletteConfig* Get() { return Instance; }
+	static IEditableToolPaletteConfigManager* GetAsConfigManager() { return Instance; }
+
+	UPROPERTY(meta=(EditorConfig))
+	TMap<FName, FEditableToolPaletteSettings> EditableToolPalettes;
+
+private:
+
+	static TObjectPtr<UModelingModeEditableToolPaletteConfig> Instance;
+};
\ No newline at end of file
diff --git a/Engine/Plugins/Editor/ModelingToolsEditorMode/Source/ModelingToolsEditorMode/Private/ModelingToolsEditorModeToolkit.cpp b/Engine/Plugins/Editor/ModelingToolsEditorMode/Source/ModelingToolsEditorMode/Private/ModelingToolsEditorModeToolkit.cpp
--- a/Engine/Plugins/Editor/ModelingToolsEditorMode/Source/ModelingToolsEditorMode/Private/ModelingToolsEditorModeToolkit.cpp
+++ b/Engine/Plugins/Editor/ModelingToolsEditorMode/Source/ModelingToolsEditorMode/Private/ModelingToolsEditorModeToolkit.cpp
@@ -39,6 +39,7 @@
 #include "SPrimaryButton.h"
 #include "Dialogs/DlgPickPath.h"
 #include "ModelingModeAssetUtils.h"
+#include "ModelingToolsEditablePaletteConfig.h"
 
 // for presets
 #include "PresetAsset.h"
@@ -102,6 +103,9 @@ namespace FModelingToolsEditorModeToolkitLocals
 
 FModelingToolsEditorModeToolkit::FModelingToolsEditorModeToolkit()
 {
+	UModelingModeEditableToolPaletteConfig::Initialize();
+	UModelingModeEditableToolPaletteConfig::Get()->LoadEditorConfig();
+
 	PresetSettings = NewObject<UPresetSettingsProperties>();
 }
 
@@ -161,15 +165,18 @@ void FModelingToolsEditorModeToolkit::RegisterPalettes()
 				GetToolkitCommands(),
 				ToolkitSections);
 
-	
-	const TSharedPtr<FEditablePalette> FavoritesPalette =
+	// We need to specify the modeling mode specific config instance because this cannot exist at the base toolkit level due to dependency issues currently
+	FGetEditableToolPaletteConfigManager GetConfigManager = FGetEditableToolPaletteConfigManager::CreateStatic(&UModelingModeEditableToolPaletteConfig::GetAsConfigManager);
+
+	FavoritesPalette =
 		MakeShareable(new FEditablePalette(
 			Commands.LoadFavoritesTools.ToSharedRef(),
-			UISettings->ToolFavorites,
 			Commands.AddToFavorites.ToSharedRef(),
-			Commands.RemoveFromFavorites.ToSharedRef()));
-	ToolkitBuilder->AddPalette(StaticCastSharedPtr<FEditablePalette>(FavoritesPalette));
+			Commands.RemoveFromFavorites.ToSharedRef(),
+			"ModelingToolsEditorModeFavoritesPalette",
+			GetConfigManager));
 	
+	ToolkitBuilder->AddPalette(StaticCastSharedPtr<FEditablePalette>(FavoritesPalette));
 
 	TArray<TSharedPtr<FUICommandInfo>> ShapesPaletteItems({
 		Commands.BeginAddBoxPrimitiveTool,
@@ -1374,7 +1381,7 @@ void FModelingToolsEditorModeToolkit::GetToolPaletteNames(TArray<FName>& Palette
 	}
 
 	// if user has provided a list of favorite tools, add that palette to the list
-	if (UISettings->ToolFavorites.Num() > 0)
+	if (FavoritesPalette && FavoritesPalette->GetPaletteCommandNames().Num() > 0)
 	{
 		PaletteNames.Insert(ModelingFavoritesTabName, 0);
 	}
@@ -1401,18 +1408,12 @@ void FModelingToolsEditorModeToolkit::BuildToolPalette(FName PaletteIndex, class
 	if (PaletteIndex == ModelingFavoritesTabName)
 	{
 		// build Favorites tool palette
-		for (FString ToolName : UISettings->ToolFavorites)
+		TArray<TSharedPtr<const FUICommandInfo>> FavoriteCommands;
+		ToolkitBuilder->GetCommandsForEditablePalette(FavoritesPalette.ToSharedRef(), FavoriteCommands);
+		
+		for (const TSharedPtr<const FUICommandInfo>& ToolCommand : FavoriteCommands)
 		{
-			bool bFound = false;
-			TSharedPtr<FUICommandInfo> FoundToolCommand = Commands.FindToolByName(ToolName, bFound);
-			if ( bFound )
-			{ 
-				ToolbarBuilder.AddToolBarButton(FoundToolCommand);
-			}
-			else
-			{
-				UE_LOG(LogTemp, Display, TEXT("ModelingMode: could not find Favorited Tool %s"), *ToolName);
-			}
+			ToolbarBuilder.AddToolBarButton(ToolCommand);
 		}
 	}
 	else if (PaletteIndex == SelectionActionsTabName)
diff --git a/Engine/Plugins/Editor/ModelingToolsEditorMode/Source/ModelingToolsEditorMode/Public/ModelingToolsEditorModeSettings.h b/Engine/Plugins/Editor/ModelingToolsEditorMode/Source/ModelingToolsEditorMode/Public/ModelingToolsEditorModeSettings.h
--- a/Engine/Plugins/Editor/ModelingToolsEditorMode/Source/ModelingToolsEditorMode/Public/ModelingToolsEditorModeSettings.h
+++ b/Engine/Plugins/Editor/ModelingToolsEditorMode/Source/ModelingToolsEditorMode/Public/ModelingToolsEditorModeSettings.h
@@ -313,9 +313,9 @@ public:
 	TArray<FString> ToolSectionOrder;
 
 	/** Tool Names listed in the array below will appear in a Favorites section at the top of the Modeling Mode Tool Palette */
-	UPROPERTY(config, EditAnywhere, Category = "Modeling Mode|UI Customization")
+	UE_DEPRECATED(5.3, "Modeling Mode favorites are now set through FEditablePalette or the Mode UI itself")
 	TArray<FString> ToolFavorites;
-
+	
 	/** Custom Section Header Colors for listed Sections in the Modeling Mode Tool Palette */
 	UPROPERTY(config, EditAnywhere, Category = "Modeling Mode|UI Customization")
 	TArray<FModelingModeCustomSectionColor> SectionColors;
diff --git a/Engine/Plugins/Editor/ModelingToolsEditorMode/Source/ModelingToolsEditorMode/Public/ModelingToolsEditorModeToolkit.h b/Engine/Plugins/Editor/ModelingToolsEditorMode/Source/ModelingToolsEditorMode/Public/ModelingToolsEditorModeToolkit.h
--- a/Engine/Plugins/Editor/ModelingToolsEditorMode/Source/ModelingToolsEditorMode/Public/ModelingToolsEditorModeToolkit.h
+++ b/Engine/Plugins/Editor/ModelingToolsEditorMode/Source/ModelingToolsEditorMode/Public/ModelingToolsEditorModeToolkit.h
@@ -186,6 +186,7 @@ private:
 	TArray<TSharedPtr<FString>> AssetLODModes;
 	TSharedPtr<STextBlock> AssetLODModeLabel;
 	TSharedPtr<STextComboBox> AssetLODMode;
+	TSharedPtr<FEditablePalette> FavoritesPalette;
 
 	bool bFirstInitializeAfterModeSetup = true;
 	bool bShowActiveSelectionActions = true;
diff --git a/Engine/Source/Developer/WidgetRegistration/Private/ToolkitBuilder.cpp b/Engine/Source/Developer/WidgetRegistration/Private/ToolkitBuilder.cpp
--- a/Engine/Source/Developer/WidgetRegistration/Private/ToolkitBuilder.cpp
+++ b/Engine/Source/Developer/WidgetRegistration/Private/ToolkitBuilder.cpp
@@ -20,15 +20,71 @@ bool FEditablePalette::IsInPalette(const FName CommandName) const
 	return this->PaletteCommandNameArray.Contains(CommandName.ToString());
 }
 
+TArray<FString> FEditablePalette::GetPaletteCommandNames() const
+{
+	return PaletteCommandNameArray;
+}
+
+void FEditablePalette::AddCommandToPalette(const FString CommandNameString)
+{
+	PaletteCommandNameArray.Add(CommandNameString);
+
+	SaveToConfig();
+	
+	OnPaletteEdited.ExecuteIfBound();
+}
+
+void FEditablePalette::RemoveCommandFromPalette(const FString CommandNameString)
+{
+	PaletteCommandNameArray.Remove(CommandNameString);
+
+	SaveToConfig();
+	
+	OnPaletteEdited.ExecuteIfBound();
+}
+
+void FEditablePalette::SaveToConfig()
+{
+	if(GetConfigManager.IsBound())
+	{
+		IEditableToolPaletteConfigManager* ConfigManager = GetConfigManager.Execute();
+		if(ConfigManager)
+		{
+			if(FEditableToolPaletteSettings* Config = ConfigManager->GetMutablePaletteConfig(EditablePaletteName))
+			{
+				Config->PaletteCommandNames = PaletteCommandNameArray;
+				ConfigManager->SavePaletteConfig(EditablePaletteName);
+			}
+		}
+	}
+}
+
+void FEditablePalette::LoadFromConfig()
+{
+	if(GetConfigManager.IsBound())
+	{
+		if(IEditableToolPaletteConfigManager* ConfigManager = GetConfigManager.Execute())
+		{
+			if(FEditableToolPaletteSettings* Config = ConfigManager->GetMutablePaletteConfig(EditablePaletteName))
+			{
+				PaletteCommandNameArray = Config->PaletteCommandNames;
+			}
+		}
+	}
+}
+
 FEditablePalette::FEditablePalette(TSharedPtr<FUICommandInfo> InLoadToolPaletteAction,
-	TArray<FString>& InPaletteCommandNameArray,
 	TSharedPtr<FUICommandInfo> InAddToPaletteAction,
-	TSharedPtr<FUICommandInfo> InRemoveFromPaletteAction) :
+	TSharedPtr<FUICommandInfo> InRemoveFromPaletteAction,
+	FName InEditablePaletteName,
+	FGetEditableToolPaletteConfigManager InGetConfigManager) :
 	FToolPalette(InLoadToolPaletteAction, {}),
 	AddToPaletteAction(InAddToPaletteAction),
 	RemoveFromPaletteAction(InRemoveFromPaletteAction),
-	PaletteCommandNameArray(InPaletteCommandNameArray)
+	EditablePaletteName(InEditablePaletteName),
+	GetConfigManager(InGetConfigManager)
 {
+	LoadFromConfig();
 }
 
 FToolkitBuilder::FToolkitBuilder(
@@ -53,8 +109,29 @@ TSharedRef<SWidget> FToolkitBuilder::CreateToolbarWidget() const
 	return ToolRegistry.GenerateWidget(VerticalToolbarElement.ToSharedRef());
 }
 
+void FToolkitBuilder::GetCommandsForEditablePalette(TSharedRef<FEditablePalette> EditablePalette, TArray<TSharedPtr<const FUICommandInfo>>& OutCommands)
+{
+	TArray<FString> CommandNames = EditablePalette->GetPaletteCommandNames();
+
+	for(const FString& CommandName : CommandNames)
+	{
+		if(TSharedPtr<const FUICommandInfo>* FoundCommand = PaletteCommandInfos.Find(CommandName))
+		{
+			if(*FoundCommand)
+			{
+				OutCommands.Add(*FoundCommand);
+			}
+			else
+			{
+				UE_LOG(LogTemp, Display, TEXT("%s: Could not find Favorited Tool %s"), *ToolbarCustomizationName.ToString(), *CommandName);
+			}
+		}
+	}
+}
+
 void FToolkitBuilder::AddPalette(TSharedPtr<FEditablePalette> Palette)
 {
+	Palette->OnPaletteEdited.BindSP(this, &FToolkitBuilder::OnEditablePaletteEdited, Palette.ToSharedRef());
 	EditablePalettesArray.Add(Palette.ToSharedRef());
 	AddPalette(StaticCastSharedRef<FToolPalette>(Palette.ToSharedRef()) );
 }
@@ -64,6 +141,7 @@ void FToolkitBuilder::AddPalette(TSharedPtr<FToolPalette> Palette)
 	for (TSharedRef<FButtonArgs> Button : Palette->PaletteActions)
 	{
 		PaletteCommandNameToButtonArgsMap.Add(Button->Command->GetCommandName().ToString(), Button);
+		PaletteCommandInfos.Add(Button->Command->GetCommandName().ToString(), Button->Command);
 	}
 	LoadCommandNameToToolPaletteMap.Add(Palette->LoadToolPaletteAction->GetCommandName().ToString(), Palette);
 
@@ -87,47 +165,52 @@ ECheckBoxState FToolkitBuilder::IsActiveToolPalette(FName CommandName) const
 				ECheckBoxState::Checked : ECheckBoxState::Unchecked;
 }
 
-void FToolkitBuilder::UpdateEditablePalette(FEditablePalette& Palette)
+void FToolkitBuilder::UpdateEditablePalette(TSharedRef<FEditablePalette> Palette)
 {
-	Palette.PaletteActions.Empty();
+	Palette->PaletteActions.Empty();
 
-	for (const FString& Key : Palette.PaletteCommandNameArray)
+	TArray<FString> PaletteComandNameArray = Palette->GetPaletteCommandNames();
+
+	for (const FString& Key : PaletteComandNameArray)
 	{
 		if (const TSharedPtr<FButtonArgs>* FoundButton = PaletteCommandNameToButtonArgsMap.Find(Key))
 		{
 			const TSharedRef<FButtonArgs> Button = (*FoundButton).ToSharedRef();
-			Palette.PaletteActions.Add(Button);
+			Palette->PaletteActions.Add(Button);
 		}
 	}
 }
 
+void FToolkitBuilder::OnEditablePaletteEdited(TSharedRef<FEditablePalette> EditablePalette)
+{
+	UpdateEditablePalette(EditablePalette);
+
+	// if the active Palette is the Palette to which we are toggling an action,
+	// recreate it to load the new state after the toggle
+	if ( ActivePalette == EditablePalette )
+	{
+		CreatePalette(EditablePalette);
+	}
+}
+
 void FToolkitBuilder::UpdateWidget()
 {
 	for (const TSharedRef<FEditablePalette>& EditablePalette : EditablePalettesArray)
 	{
-		UpdateEditablePalette(*EditablePalette);		
+		UpdateEditablePalette(EditablePalette);		
 	}
 	
 }
 
 void FToolkitBuilder::ToggleCommandInPalette(TSharedRef<FEditablePalette> Palette, FString CommandNameString)
 {
-	if (!Palette->PaletteCommandNameArray.Contains(CommandNameString))
+	if (!Palette->IsInPalette(FName(CommandNameString)))
 	{	
-		Palette->PaletteCommandNameArray.Add(CommandNameString);
+		Palette->AddCommandToPalette(CommandNameString);
 	}
 	else
 	{
-		Palette->PaletteCommandNameArray.Remove(CommandNameString);		
-	}
-
-	UpdateEditablePalette(*Palette);
-
-	// if the active Palette is the Palette to which we are toggling an action,
-	// recreate it to load the new state after the toggle
-	if ( ActivePalette == Palette )
-	{
-		CreatePalette(Palette);
+		Palette->RemoveCommandFromPalette(CommandNameString);		
 	}
 }
 
diff --git a/Engine/Source/Developer/WidgetRegistration/Public/ToolkitBuilder.h b/Engine/Source/Developer/WidgetRegistration/Public/ToolkitBuilder.h
--- a/Engine/Source/Developer/WidgetRegistration/Public/ToolkitBuilder.h
+++ b/Engine/Source/Developer/WidgetRegistration/Public/ToolkitBuilder.h
@@ -8,6 +8,7 @@
 #include "Framework/Commands/UICommandInfo.h"
 #include "InteractiveToolManager.h"
 #include "ToolElementRegistry.h"
+#include "ToolkitBuilderConfig.h"
 #include "ToolMenus.h"
 #include "Widgets/SBoxPanel.h"
 #include "Widgets/Layout/SSplitter.h"
@@ -46,32 +47,55 @@ struct WIDGETREGISTRATION_API FToolPalette : TSharedFromThis<FToolPalette>
 	TArray<TSharedRef<FButtonArgs>> PaletteActions;
 
 	/** The FUICommandList associated with this Palette */
-	TSharedPtr<FUICommandList>PaletteActionsCommandList;
+	TSharedPtr<FUICommandList> PaletteActionsCommandList;
 };
 
 /** An FToolPalette to which you can add and remove actions */
 struct WIDGETREGISTRATION_API FEditablePalette : public FToolPalette
 {
 	FEditablePalette(TSharedPtr<FUICommandInfo> InLoadToolPaletteAction,
-		TArray<FString>& InPaletteCommandNameArray,
 		TSharedPtr<FUICommandInfo> InAddToPaletteAction,
-		TSharedPtr<FUICommandInfo> InRemoveFromPaletteAction);
+		TSharedPtr<FUICommandInfo> InRemoveFromPaletteAction,
+		FName InEditablePaletteName = FName(),
+		FGetEditableToolPaletteConfigManager InGetConfigManager = FGetEditableToolPaletteConfigManager());
 
 	/** The FUICommandInfo which adds an action to this palette */
 	const TSharedPtr<FUICommandInfo> AddToPaletteAction;
 	
 	/** The FUICommandInfo which removes an action to this palette */
 	const TSharedPtr<FUICommandInfo> RemoveFromPaletteAction;
-	
-	/** The TArray of Command names that are the current FuiCommandInfo actions in this Palette */
-	TArray<FString>& PaletteCommandNameArray;
 
+	/** Delegate used by FToolkitBuilder that is called when an item is added/removed from the palette */
+	FSimpleDelegate OnPaletteEdited;
+	
 	/**
 	 * Given a reference to a FUICommandInfo, returns whether it is in the current Palette
 	 *
 	 * @param CommandName the name of the FUICommandInfo queried as to whether it is in the Palette
 	 */
 	bool IsInPalette(const FName CommandName) const;
+
+	TArray<FString> GetPaletteCommandNames() const;
+
+	void AddCommandToPalette(const FString CommandNameString);
+
+	void RemoveCommandFromPalette(const FString CommandNameString);
+
+protected:
+
+	void SaveToConfig();
+
+	void LoadFromConfig();
+	
+protected:
+	/** The TArray of Command names that are the current FuiCommandInfo actions in this Palette */
+	TArray<FString> PaletteCommandNameArray;
+
+	/** The (unique) name attached to this palette, enables saving the palette contents into config if provided */
+	FName EditablePaletteName;
+
+	/** Delegate used to check if we have a config manager and get it */
+	FGetEditableToolPaletteConfigManager GetConfigManager;
 };
 
 
@@ -147,6 +171,7 @@ public:
 
 	FText GetActiveToolDisplayName() const;
 
+	void GetCommandsForEditablePalette(TSharedRef<FEditablePalette> EditablePalette, TArray<TSharedPtr<const FUICommandInfo>>& OutCommands);
 	
 private:
 
@@ -165,6 +190,9 @@ private:
 	/** The map of the command name to the ButtonArgs for it  */
 	TMap<FString, TSharedPtr<FToolPalette>> LoadCommandNameToToolPaletteMap;
 
+	/** Map of command name to the actual command for all commands belonging to this palette */
+	TMap<FString, TSharedPtr<const FUICommandInfo>> PaletteCommandInfos;
+
 	/** A TSharedPointer to the FUICommandList for the FUICommandInfos which load a tool palette */
 	TSharedPtr<FUICommandList> LoadToolPaletteCommandList;
 
@@ -187,7 +215,9 @@ private:
 	TMap<FName, TSharedPtr<FToolBarBuilder>> LoadCommandNameToPaletteToolbarBuilderMap;
 
 	/** Updates the Editable Palette with any commands that are on it */
-	void UpdateEditablePalette(FEditablePalette& EditablePalette);
+	void UpdateEditablePalette(TSharedRef<FEditablePalette> EditablePalette);
+
+	void OnEditablePaletteEdited(TSharedRef<FEditablePalette> EditablePalette);
 
 	/** Resets the tool palette widget, on which the buttons for the currently chosen toolset are shown */
 	void ResetToolPaletteWidget();
diff --git a/Engine/Source/Developer/WidgetRegistration/Public/ToolkitBuilderConfig.h b/Engine/Source/Developer/WidgetRegistration/Public/ToolkitBuilderConfig.h
new file mode 100644
--- /dev/null
+++ b/Engine/Source/Developer/WidgetRegistration/Public/ToolkitBuilderConfig.h
@@ -0,0 +1,34 @@
+// Copyright Epic Games, Inc. All Rights Reserved.
+
+#pragma once
+
+#include "Containers/UnrealString.h"
+
+#include "ToolkitBuilderConfig.generated.h"
+
+// The actual struct that contains the properties saved for an FEditableToolPalette
+USTRUCT()
+struct WIDGETREGISTRATION_API FEditableToolPaletteSettings
+{
+	GENERATED_BODY()
+	
+public:
+	
+	UPROPERTY()
+	TArray<FString> PaletteCommandNames;
+};
+
+// Interface class for FEditableToolPalette to interact with EditorConfig which is an Editor only module currently
+class IEditableToolPaletteConfigManager
+{
+public:
+
+	virtual FEditableToolPaletteSettings* GetMutablePaletteConfig(const FName& InstanceName) = 0;
+	virtual const FEditableToolPaletteSettings* GetConstPaletteConfig(const FName& InstanceName) = 0;
+	virtual void SavePaletteConfig(const FName& InstanceName) = 0;
+	virtual ~IEditableToolPaletteConfigManager() { }
+};
+
+// Delegate passed into FEditableToolPalette that is used to get the actual config manager instance
+DECLARE_DELEGATE_RetVal(IEditableToolPaletteConfigManager*, FGetEditableToolPaletteConfigManager)
+
