Updates for ToolkitBuilder to allow category labels based on editor preferences, and remove the category title.

diff --git a/Engine/Plugins/Editor/ModelingToolsEditorMode/Source/ModelingToolsEditorMode/Private/ModelingToolsEditorModeToolkit.cpp b/Engine/Plugins/Editor/ModelingToolsEditorMode/Source/ModelingToolsEditorMode/Private/ModelingToolsEditorModeToolkit.cpp
--- a/Engine/Plugins/Editor/ModelingToolsEditorMode/Source/ModelingToolsEditorMode/Private/ModelingToolsEditorModeToolkit.cpp
+++ b/Engine/Plugins/Editor/ModelingToolsEditorMode/Source/ModelingToolsEditorMode/Private/ModelingToolsEditorModeToolkit.cpp
@@ -214,12 +214,14 @@ void FModelingToolsEditorModeToolkit::RegisterPalettes()
 	UEdMode* ScriptableMode = GetScriptableEditorMode().Get();
 	UModelingToolsModeCustomizationSettings* UISettings = GetMutableDefault<UModelingToolsModeCustomizationSettings>();
 	UModelingToolsEditorModeSettings* ModelingModeSettings = GetMutableDefault<UModelingToolsEditorModeSettings>();
-
+	UISettings->OnSettingChanged().AddSP(SharedThis(this), &FModelingToolsEditorModeToolkit::UpdateCategoryButtonLabelVisibility);
+	
 	ToolkitSections = MakeShared<FToolkitSections>();
 	ToolkitBuilder = MakeShared<FToolkitBuilder>(
 				ScriptableMode->GetModeInfo().ToolbarCustomizationName, 
 				GetToolkitCommands(),
 				ToolkitSections);
+	ToolkitBuilder->SetCategoryButtonLabelVisibility(UISettings->bShowCategoryButtonLabels);
 
 	// We need to specify the modeling mode specific config instance because this cannot exist at the base toolkit level due to dependency issues currently
 	FGetEditableToolPaletteConfigManager GetConfigManager = FGetEditableToolPaletteConfigManager::CreateStatic(&UModelingModeEditableToolPaletteConfig::GetAsConfigManager);
@@ -2265,6 +2267,13 @@ void FModelingToolsEditorModeToolkit::SelectNewAssetPath() const
 	}
 }
 
+void FModelingToolsEditorModeToolkit::UpdateCategoryButtonLabelVisibility(UObject* Obj, FPropertyChangedEvent& ChangeEvent)
+{
+	UModelingToolsModeCustomizationSettings* UISettings = GetMutableDefault<UModelingToolsModeCustomizationSettings>();
+	ToolkitBuilder->SetCategoryButtonLabelVisibility(UISettings != nullptr  ? UISettings->bShowCategoryButtonLabels : true);
+	ToolkitBuilder->RefreshCategoryToolbarWidget();
+}
+
 
 FText FModelingToolsEditorModeToolkit::GetRestrictiveModeAutoGeneratedAssetPathText() const
 {
diff --git a/Engine/Plugins/Editor/ModelingToolsEditorMode/Source/ModelingToolsEditorMode/Public/ModelingToolsEditorModeSettings.h b/Engine/Plugins/Editor/ModelingToolsEditorMode/Source/ModelingToolsEditorMode/Public/ModelingToolsEditorModeSettings.h
--- a/Engine/Plugins/Editor/ModelingToolsEditorMode/Source/ModelingToolsEditorMode/Public/ModelingToolsEditorModeSettings.h
+++ b/Engine/Plugins/Editor/ModelingToolsEditorMode/Source/ModelingToolsEditorMode/Public/ModelingToolsEditorModeSettings.h
@@ -339,7 +339,12 @@ public:
 	TArray<FModelingModeAssetCollectionSet> BrushAlphaSets;
 
 
-
+	/**
+	 * If true, the category labels will be shown on the toolbar buttons, else they will be hidden
+	 */
+	UPROPERTY(config, EditAnywhere, Category = "Modeling Mode|UI Customization")
+	bool bShowCategoryButtonLabels = true;
+	
 public:
 
 	// saved-state for various mode settings that are configured via UI toggles/etc, and not exposed in settings dialog
diff --git a/Engine/Plugins/Editor/ModelingToolsEditorMode/Source/ModelingToolsEditorMode/Public/ModelingToolsEditorModeToolkit.h b/Engine/Plugins/Editor/ModelingToolsEditorMode/Source/ModelingToolsEditorMode/Public/ModelingToolsEditorModeToolkit.h
--- a/Engine/Plugins/Editor/ModelingToolsEditorMode/Source/ModelingToolsEditorMode/Public/ModelingToolsEditorModeToolkit.h
+++ b/Engine/Plugins/Editor/ModelingToolsEditorMode/Source/ModelingToolsEditorMode/Public/ModelingToolsEditorModeToolkit.h
@@ -177,6 +177,10 @@ private:
 	bool IsPresetEnabled() const;
 	void ClearPresetComboList();
 
+	/**
+	 *  Updates the category button labels' visibility based off of the current editor preference for it
+	 */
+	void UpdateCategoryButtonLabelVisibility(UObject* Obj, FPropertyChangedEvent& ChangeEvent);
 
 	TArray<TSharedPtr<FString>> AssetLODModes;
 	TSharedPtr<STextBlock> AssetLODModeLabel;
diff --git a/Engine/Source/Developer/WidgetRegistration/Private/ToolkitBuilder.cpp b/Engine/Source/Developer/WidgetRegistration/Private/ToolkitBuilder.cpp
--- a/Engine/Source/Developer/WidgetRegistration/Private/ToolkitBuilder.cpp
+++ b/Engine/Source/Developer/WidgetRegistration/Private/ToolkitBuilder.cpp
@@ -87,6 +87,29 @@ FEditablePalette::FEditablePalette(TSharedPtr<FUICommandInfo> InLoadToolPaletteA
 	LoadFromConfig();
 }
 
+FToolkitBuilder::FToolkitBuilder(const FToolkitBuilderArgs& Args) :
+	FToolElementRegistrationArgs(EToolElement::Toolkit),
+	ToolbarCustomizationName(Args.ToolbarCustomizationName),
+	ToolkitCommandList(Args.ToolkitCommandList),
+	ToolkitSections(Args.ToolkitSections),
+	SelectedCategoryTitleVisibility(Args.SelectedCategoryTitleVisibility)
+{
+	SetCategoryButtonLabelVisibility(Args.bShowCategoryButtonLabels);
+	ResetWidget();
+}
+
+FToolkitBuilder::~FToolkitBuilder()
+{
+	for (const TSharedRef<FToolElement>& PaletteElement : ToolPaletteElementArray)
+	{
+		ToolRegistry.UnregisterElement(PaletteElement);
+	}
+	if (VerticalToolbarElement.IsValid())
+	{
+		ToolRegistry.UnregisterElement(VerticalToolbarElement.ToSharedRef());
+	}
+}
+
 FToolkitBuilder::FToolkitBuilder(
 	FName InToolbarCustomizationName,
 	TSharedPtr<FUICommandList> InToolkitCommandList,
@@ -219,6 +242,69 @@ bool FToolkitBuilder::HasActivePalette() const
 	return ActivePalette != nullptr;
 }
 
+void FToolkitBuilder::InitializeCategoryToolbar(bool InitLoadToolPaletteMap)
+{
+	Style = FToolkitStyle::Get().GetWidgetStyle<FToolkitWidgetStyle>("FToolkitWidgetStyle");
+	LoadToolPaletteCommandList = MakeShareable(new FUICommandList);
+	LoadPaletteToolBarBuilder = MakeShared<FVerticalToolBarBuilder>(LoadToolPaletteCommandList, FMultiBoxCustomization::None, TSharedPtr<FExtender>(), true);
+	LoadPaletteToolBarBuilder->SetLabelVisibility( CategoryButtonLabelVisibility );
+	EditablePalettesArray.Reset();
+
+	if (InitLoadToolPaletteMap)
+	{
+		LoadCommandNameToPaletteToolbarBuilderMap.Reset();
+	}
+	else
+	{
+		TArray<FString> Keys;
+		LoadCommandNameToToolPaletteMap.GetKeys(Keys);
+		for ( const FString& Name : Keys)
+		{
+			const TSharedPtr<FToolPalette>* Palette = LoadCommandNameToToolPaletteMap.Find(Name);
+			if (Palette)
+			{
+				AddPalette(*Palette);			
+			}
+		}
+	}
+
+}
+
+void FToolkitBuilder::InitCategoryToolbarContainerWidget()
+{
+	if (!CategoryToolbarVBox.IsValid())
+	{
+		CategoryToolbarVBox = SNew(SVerticalBox);		
+	}
+	else
+	{
+		CategoryToolbarVBox->ClearChildren();
+	}
+	CategoryToolbarVBox->AddSlot()
+	.Padding(0.f)
+	[
+		CreateToolbarWidget()
+	];
+}
+
+void FToolkitBuilder::RefreshCategoryToolbarWidget()
+{
+	FToolElementRegistrationKey Key = FToolElementRegistrationKey(ToolbarCustomizationName, EToolElement::Toolbar);
+	VerticalToolbarElement = ToolRegistry.GetToolElementSP(Key);
+	const TSharedRef<FToolbarRegistrationArgs> VerticalToolbarRegistrationArgs = MakeShareable<FToolbarRegistrationArgs>(
+		new FToolbarRegistrationArgs(LoadPaletteToolBarBuilder.ToSharedRef()));
+	
+	if (!VerticalToolbarElement.IsValid())
+	{
+		VerticalToolbarElement = MakeShareable(new FToolElement
+			(ToolbarCustomizationName,
+			VerticalToolbarRegistrationArgs));
+		ToolRegistry.RegisterElement(VerticalToolbarElement.ToSharedRef());
+	}
+
+	VerticalToolbarElement->SetRegistrationArgs(VerticalToolbarRegistrationArgs);
+	InitCategoryToolbarContainerWidget();
+}
 
 void FToolkitBuilder::TogglePalette(TSharedPtr<FToolPalette> Palette)
 {
@@ -267,6 +353,7 @@ void FToolkitBuilder::CreatePalette(TSharedPtr<FToolPalette> Palette)
 	Element->SetRegistrationArgs(RegistrationArgs.ToSharedRef());
 	
 	LoadCommandNameToPaletteToolbarBuilderMap.Add(Palette->LoadToolPaletteAction->GetCommandName(), PaletteToolbarBuilder.ToSharedRef());
+	ToolPaletteElementArray.Add(Element.ToSharedRef());
 	
 	PaletteToolbarBuilder->SetStyle(&FAppStyle::Get(), "SlimPaletteToolBar");
 
@@ -295,6 +382,7 @@ void FToolkitBuilder::CreatePaletteWidget(FToolPalette& Palette, FToolElement& E
 		SNew(SBorder)
 		.Padding(Style.TitlePadding)
 		.VAlign(VAlign_Center)
+		.Visibility( SelectedCategoryTitleVisibility )
 		.BorderImage(&Style.TitleBackgroundBrush)
 		.HAlign(HAlign_Left)
 		[
@@ -344,28 +432,8 @@ TSharedRef<SWidget> FToolkitBuilder::GetContextMenuContent(const FName CommandNa
 
 void FToolkitBuilder::ResetWidget()
 {
-	Style = FToolkitStyle::Get().GetWidgetStyle<FToolkitWidgetStyle>("FToolkitWidgetStyle");
-	LoadToolPaletteCommandList = MakeShareable(new FUICommandList);
-	LoadPaletteToolBarBuilder = MakeShared<FVerticalToolBarBuilder>(LoadToolPaletteCommandList, FMultiBoxCustomization::None, TSharedPtr<FExtender>(), true);
+	InitializeCategoryToolbar(true);
 	ToolPaletteWidget = SNew(SVerticalBox);
-
-	LoadCommandNameToPaletteToolbarBuilderMap.Reset();
-	EditablePalettesArray.Reset();
-
-	FToolElementRegistrationKey Key = FToolElementRegistrationKey(ToolbarCustomizationName, EToolElement::Toolbar);
-	VerticalToolbarElement = ToolRegistry.GetToolElementSP(Key);
-	const TSharedRef<FToolbarRegistrationArgs> VerticalToolbarRegistrationArgs = MakeShareable<FToolbarRegistrationArgs>(
-		new FToolbarRegistrationArgs(LoadPaletteToolBarBuilder.ToSharedRef()));
-	
-	if (!VerticalToolbarElement)
-	{
-		VerticalToolbarElement = MakeShareable(new FToolElement
-			(ToolbarCustomizationName,
-			VerticalToolbarRegistrationArgs));
-		ToolRegistry.RegisterElement(VerticalToolbarElement.ToSharedRef());
-	}
-
-	VerticalToolbarElement->SetRegistrationArgs(VerticalToolbarRegistrationArgs);
 }
 
 void FToolkitBuilder::ResetToolPaletteWidget()
@@ -417,6 +485,7 @@ EVisibility FToolkitBuilder::GetActiveToolTitleVisibility() const
 
 void FToolkitBuilder::DefineWidget()
 {
+	RefreshCategoryToolbarWidget();
 	ToolkitWidgetVBox = SNew(SVerticalBox);
 	ToolkitWidgetHBox = 
 		SNew(SSplitter)
@@ -424,7 +493,7 @@ void FToolkitBuilder::DefineWidget()
 		.Resizable(false)
 		.SizeRule(SSplitter::SizeToContent)
 			[
-				CreateToolbarWidget()
+				CategoryToolbarVBox.ToSharedRef()
 			]
 
 		+ SSplitter::Slot()
@@ -529,4 +598,15 @@ void FToolkitBuilder::DefineWidget()
 	}
 }
 
+void FToolkitBuilder::SetCategoryButtonLabelVisibility(EVisibility Visibility)
+{
+	CategoryButtonLabelVisibility = Visibility;
+	InitializeCategoryToolbar();
+}
+
+void FToolkitBuilder::SetCategoryButtonLabelVisibility(bool bIsCategoryButtonLabelVisible)
+{
+	SetCategoryButtonLabelVisibility(bIsCategoryButtonLabelVisible ? EVisibility::Visible : EVisibility::Collapsed);
+}
+
 #undef LOCTEXT_NAMESPACE
\ No newline at end of file
diff --git a/Engine/Source/Developer/WidgetRegistration/Public/ToolkitBuilder.h b/Engine/Source/Developer/WidgetRegistration/Public/ToolkitBuilder.h
--- a/Engine/Source/Developer/WidgetRegistration/Public/ToolkitBuilder.h
+++ b/Engine/Source/Developer/WidgetRegistration/Public/ToolkitBuilder.h
@@ -6,7 +6,6 @@
 #include "IDetailsView.h"
 #include "Framework/MultiBox/SToolBarButtonBlock.h"
 #include "Framework/Commands/UICommandInfo.h"
-#include "InteractiveToolManager.h"
 #include "ToolElementRegistry.h"
 #include "ToolkitBuilderConfig.h"
 #include "ToolMenus.h"
@@ -99,6 +98,33 @@ protected:
 	FGetEditableToolPaletteConfigManager GetConfigManager;
 };
 
+/*
+ * A simple struct to carry initialization data for an FToolkitBuilder
+ */
+struct FToolkitBuilderArgs
+{
+	FToolkitBuilderArgs(const FName& InToolbarCustomizationName) : ToolbarCustomizationName(InToolbarCustomizationName) 
+	{
+	}
+	
+	/** Name of the toolbar this mode uses and can be used by external systems to customize that mode toolbar */
+	const FName& ToolbarCustomizationName;
+
+	/** A TSharedPointer to the FUICommandList for the current mode */
+	TSharedPtr<FUICommandList> ToolkitCommandList;
+
+	/** The FToolkitSections which holds the sections defined for this Toolkit */
+	TSharedPtr<FToolkitSections> ToolkitSections;
+
+	/** If bShowCategoryButtonLabels == true, the category button
+	* labels should be, else they are not displayed  */	
+	bool bShowCategoryButtonLabels;
+
+	/** If SelectedCategoryTitleVisibility == EVisibility::Visible, the selected
+	 * category title is visible, else it is not displayed  */	
+	EVisibility SelectedCategoryTitleVisibility;
+};
+
 
 /**
  * The FToolElementRegistrationArgs which is specified for Toolkits
@@ -107,9 +133,53 @@ class WIDGETREGISTRATION_API FToolkitBuilder : public FToolElementRegistrationAr
 {
 public:
 	FToolkitBuilder(
-		FName InToolbarCustomizationName,
-		TSharedPtr<FUICommandList> InToolkitCommandList,
-		TSharedPtr<FToolkitSections> InToolkitSections);
+		FName ToolbarCustomizationName,
+		TSharedPtr<FUICommandList> ToolkitCommandList,
+		TSharedPtr<FToolkitSections> ToolkitSections);
+
+	FToolkitBuilder(const FToolkitBuilderArgs& Args);
+	
+	virtual  ~FToolkitBuilder() override;
+	
+	/*
+	 * Initializes the data to necessary to build the category toolbar
+	 *
+	 * @param InitLoadToolPaletteMap If true, the map which holds the load
+	 * action names as keys and the ToolPalettes as value will be cleared and available
+	 * to add new data when the method completes. Otherwise, the map will keep
+	 * the load action name to ToolPalette mappings that it was initially set up with. 
+	 */
+	void InitializeCategoryToolbar(bool InitLoadToolPaletteMap = false);
+
+	/**
+	 * Initializes the category toolbar container VBox, and the children inside it.
+	 * On any repeat calls, the SVerticalBox created on the first pass will be
+	 * emptied and the children repopulated.
+	 */
+	void InitCategoryToolbarContainerWidget();
+
+	/**
+	 * Sets category button label visibility to Visiblity. It also reinitializes the
+	 * category toolbar data, as the toolbar's LabelVisibility member is now stale.
+	 *
+	 * @param Visibility If Visibility == EVisibility::Collapsed, the category button labels
+	 * will be shown, else they will not be shown.
+	 */
+	void SetCategoryButtonLabelVisibility(EVisibility Visibility);
+
+	/**
+	 * Sets category button label visibility to Visiblity. It also reinitializes the
+	 * category toolbar data, as the toolbar's LabelVisibility member is now stale.
+	 *
+	 * @param bIsCategoryButtonLabelVisible If bIsCategoryButtonLabelVisible == true,
+	 * the category button labels will be shown, else they will not be shown.
+	 */
+	void SetCategoryButtonLabelVisibility(bool bIsCategoryButtonLabelVisible);
+
+	/**
+	 * RefreshCategoryToolbarWidget refreshes the UI display of the category toolbar 
+	 */
+	void RefreshCategoryToolbarWidget();
 
 	/**
 	 * Adds the FToolPalette Palette to this FToolkitBuilder
@@ -266,6 +336,13 @@ private:
 
 	void DefineWidget();
 
+	/** The SVerticalBox which contains the category toolbar */
+	TSharedPtr<SVerticalBox> CategoryToolbarVBox;
+
+	/** The array of Tool Palettes on display. This array is used to keep a handle
+	 * on all of the created FToolElements so that we can unregister them upon destruction */
+	TArray<TSharedRef<FToolElement>> ToolPaletteElementArray;
+
 	/** The SVerticalBox which holds all but the vertical toolbar in a Toolkit */
 	TSharedPtr<SVerticalBox> ToolkitWidgetVBox;
 
@@ -281,4 +358,11 @@ private:
 	/** The current FToolkitWidgetStyle */
 	FToolkitWidgetStyle Style;
 
+	/** If SelectedCategoryTitleVisibility == EVisibility::Visible, the selected
+	 * category title is visible, else it is not displayed  */	
+	EVisibility SelectedCategoryTitleVisibility;
+
+	/** If CategoryButtonLabelVisibility == EVisibility::Visible, the category button
+	 * labels are visible, else they are not displayed  */	
+	EVisibility CategoryButtonLabelVisibility;
 };
\ No newline at end of file
diff --git a/Engine/Source/Editor/EditorStyle/Private/StarshipStyle.cpp b/Engine/Source/Editor/EditorStyle/Private/StarshipStyle.cpp
--- a/Engine/Source/Editor/EditorStyle/Private/StarshipStyle.cpp
+++ b/Engine/Source/Editor/EditorStyle/Private/StarshipStyle.cpp
@@ -1738,7 +1738,9 @@ void FStarshipEditorStyle::FStyle::SetupGeneralStyles()
 		FVerticalToolBarLabelStyle.SetOverflowPolicy(ETextOverflowPolicy::Ellipsis);
 		
 		FVerticalToolBarStyle.SetLabelStyle(FVerticalToolBarLabelStyle);
-		FVerticalToolBarStyle.SetLabelPadding(FMargin(0.0f, 4.0f, 0.0f, 0.0f));
+		FVerticalToolBarStyle.SetLabelPadding(FMargin(4.0f, 2.0f, 4.0f, 4.0f));
+		FVerticalToolBarStyle.SetButtonContentMaxWidth( 64.0f );
+		FVerticalToolBarStyle.SetButtonContentFillWidth( 1.0f );
 
 		FVerticalToolBarStyle.SetButtonPadding(       FMargin(0.0f, 0.0f, 0.0f, 0.0f));
 		//not this
@@ -1749,7 +1751,8 @@ void FStarshipEditorStyle::FStyle::SetupGeneralStyles()
 		FVerticalToolBarStyle.ButtonStyle.SetNormalPadding(FMargin(12.0f, 6.0f));
 		FVerticalToolBarStyle.ButtonStyle.SetPressedPadding(FMargin(12.0f, 6.0f));
 		FVerticalToolBarStyle.SetExpandBrush(IMAGE_BRUSH("Icons/toolbar_expand_16x", Icon8x8));
-		FVerticalToolBarStyle.SetIconPadding(FMargin(8.0f, 8.0f));
+		FVerticalToolBarStyle.SetIconPaddingWithCollapsedLabel(FMargin(8.0f, 8.0f));
+		FVerticalToolBarStyle.SetIconPaddingWithVisibleLabel(FMargin(8.0f, 8.0f, 8.0f, 2.0f));
  
 		FCheckBoxStyle CheckBoxStyle = FCheckBoxStyle(FVerticalToolBarStyle.ToggleButton)
 			.SetCheckedImage(FSlateRoundedBoxBrush(FStyleColors::Primary, 4.f, FLinearColor(0, 0, 0, .8), 0.5))
diff --git a/Engine/Source/Runtime/Slate/Private/Framework/MultiBox/SToolBarButtonBlock.cpp b/Engine/Source/Runtime/Slate/Private/Framework/MultiBox/SToolBarButtonBlock.cpp
--- a/Engine/Source/Runtime/Slate/Private/Framework/MultiBox/SToolBarButtonBlock.cpp
+++ b/Engine/Source/Runtime/Slate/Private/Framework/MultiBox/SToolBarButtonBlock.cpp
@@ -264,32 +264,43 @@ void SToolBarButtonBlock::BuildMultiBlockWidget(const ISlateStyle* StyleSet, con
 	}
 	else
 	{
+		const FMargin IconPadding = LabelVisibility.Get() == EVisibility::Collapsed ? ToolBarStyle.IconPaddingWithCollapsedLabel :
+						ToolBarStyle.IconPaddingWithVisibleLabel;
+		const TSharedRef ContentVBox = 
+			SNew(SVerticalBox)
+					// Icon image
+					+ SVerticalBox::Slot()
+					.AutoHeight()
+					.Padding(IconPadding)
+					.HAlign(HAlign_Center)	// Center the icon horizontally, so that large labels don't stretch out the artwork
+					[
+						IconWidget
+					]
+					// Label text
+					+ SVerticalBox::Slot().AutoHeight()
+				.Padding(ToolBarStyle.LabelPadding)
+					.HAlign(HAlign_Center)	// Center the label text horizontally
+					[
+						SNew(STextBlock)
+					.Visibility(LabelVisibility)
+						.Text(ActualLabel)
+						.TextStyle(&ToolBarStyle.LabelStyle)	// Smaller font for tool tip labels
+					];
+
 		ButtonContent =
 			SNew(SHorizontalBox)
 			.AddMetaData<FTagMetaData>(FTagMetaData(TutorialHighlightName))
-			+ SHorizontalBox::Slot()
-			.FillWidth(1)
-			.VAlign(VAlign_Center)
+			+ ( MultiBox->GetType() == EMultiBoxType::VerticalToolBar ?
+				SHorizontalBox::Slot()
+			.MaxWidth(ToolBarStyle.ButtonContentMaxWidth)
+			.SizeParam(FStretch())
+			.VAlign(VAlign_Center) :
+			SHorizontalBox::Slot()
+			.FillWidth(ToolBarStyle.ButtonContentFillWidth)
+			.AutoWidth()
+			.VAlign(VAlign_Center))
 			[
-				SNew(SVerticalBox)
-				// Icon image
-				+ SVerticalBox::Slot()
-				.AutoHeight()
-				.Padding(ToolBarStyle.IconPadding)
-				.HAlign(HAlign_Center)	// Center the icon horizontally, so that large labels don't stretch out the artwork
-				[
-					IconWidget
-				]
-				// Label text
-				+ SVerticalBox::Slot().AutoHeight()
-				.Padding(ToolBarStyle.LabelPadding)
-				.HAlign(HAlign_Center)	// Center the label text horizontally
-				[
-					SNew(STextBlock)
-					.Visibility(LabelVisibility)
-					.Text(ActualLabel)
-					.TextStyle(&ToolBarStyle.LabelStyle)	// Smaller font for tool tip labels
-				]
+				ContentVBox
 			];
 		}
 
@@ -342,6 +353,7 @@ void SToolBarButtonBlock::BuildMultiBlockWidget(const ISlateStyle* StyleSet, con
 		const TSharedPtr<SWidget> CheckBox = SNew(SCheckBox)
 						// Use the tool bar style for this check box
 						.Style(CheckStyle)
+						.CheckBoxContentUsesAutoWidth(false)
 						.IsFocusable(bIsFocusable)
 						.ToolTip( FMultiBoxSettings::ToolTipConstructor.Execute( ActualToolTip, nullptr, Action.Pin()))		
 						.OnCheckStateChanged(this, &SToolBarButtonBlock::OnCheckStateChanged )
diff --git a/Engine/Source/Runtime/SlateCore/Public/Styling/ToolBarStyle.h b/Engine/Source/Runtime/SlateCore/Public/Styling/ToolBarStyle.h
--- a/Engine/Source/Runtime/SlateCore/Public/Styling/ToolBarStyle.h
+++ b/Engine/Source/Runtime/SlateCore/Public/Styling/ToolBarStyle.h
@@ -128,4 +128,21 @@ struct SLATECORE_API FToolBarStyle : public FSlateWidgetStyle
 	UPROPERTY(EditAnywhere, BlueprintReadWrite, Category=Appearance)
 	bool bShowLabels;
 	FToolBarStyle& SetShowLabels(bool bInShowLabels) { bShowLabels = bInShowLabels; return *this; }
+	
+	UPROPERTY(EditAnywhere, BlueprintReadWrite, Category=Appearance)
+	float ButtonContentMaxWidth;
+	FToolBarStyle& SetButtonContentMaxWidth(float InButtonContentMaxWidth) { ButtonContentMaxWidth = InButtonContentMaxWidth; return *this; }
+
+	UPROPERTY(EditAnywhere, BlueprintReadWrite, Category=Appearance)
+	float ButtonContentFillWidth;
+	FToolBarStyle& SetButtonContentFillWidth(float InButtonContentFillWidth) { ButtonContentFillWidth = InButtonContentFillWidth; return *this; }
+	
+	UPROPERTY(EditAnywhere, BlueprintReadWrite, Category = Appearance)
+	FMargin IconPaddingWithVisibleLabel;
+	FToolBarStyle& SetIconPaddingWithVisibleLabel(const FMargin& InIconPaddingWithVisibleLabel) { IconPaddingWithVisibleLabel = InIconPaddingWithVisibleLabel; return *this; }
+	
+	UPROPERTY(EditAnywhere, BlueprintReadWrite, Category = Appearance)
+	FMargin IconPaddingWithCollapsedLabel;
+	FToolBarStyle& SetIconPaddingWithCollapsedLabel(const FMargin& InIconPaddingWithCollapsedLabel) { IconPaddingWithCollapsedLabel = InIconPaddingWithCollapsedLabel; return *this; }
+	
 };
