Add ability to dynamically toggle visibility of ToolkitBuilder command buttons, and use in Modeling Mode to hide button panel when inside a Tool

diff --git a/Engine/Plugins/Editor/ModelingToolsEditorMode/Source/ModelingToolsEditorMode/Private/ModelingToolsEditorModeToolkit.cpp b/Engine/Plugins/Editor/ModelingToolsEditorMode/Source/ModelingToolsEditorMode/Private/ModelingToolsEditorModeToolkit.cpp
--- a/Engine/Plugins/Editor/ModelingToolsEditorMode/Source/ModelingToolsEditorMode/Private/ModelingToolsEditorModeToolkit.cpp
+++ b/Engine/Plugins/Editor/ModelingToolsEditorMode/Source/ModelingToolsEditorMode/Private/ModelingToolsEditorModeToolkit.cpp
@@ -2045,6 +2045,7 @@ void FModelingToolsEditorModeToolkit::OnToolStarted(UInteractiveToolManager* Man
 	if (HasToolkitBuilder())
 	{
 		ToolkitBuilder->SetActiveToolDisplayName(ActiveToolName);
+		ToolkitBuilder->SetActivePaletteCommandsVisibility(EVisibility::Collapsed);
 	}
 
 	// try to update icon
@@ -2086,6 +2087,7 @@ void FModelingToolsEditorModeToolkit::OnToolEnded(UInteractiveToolManager* Manag
 	if (HasToolkitBuilder())
 	{
 		ToolkitBuilder->SetActiveToolDisplayName(FText::GetEmpty());		
+		ToolkitBuilder->SetActivePaletteCommandsVisibility(EVisibility::Visible);
 	}
 
 	ModeHeaderArea->SetVisibility(EVisibility::Visible);
diff --git a/Engine/Source/Developer/WidgetRegistration/Private/ToolkitBuilder.cpp b/Engine/Source/Developer/WidgetRegistration/Private/ToolkitBuilder.cpp
--- a/Engine/Source/Developer/WidgetRegistration/Private/ToolkitBuilder.cpp
+++ b/Engine/Source/Developer/WidgetRegistration/Private/ToolkitBuilder.cpp
@@ -394,11 +394,14 @@ void FToolkitBuilder::CreatePaletteWidget(FToolPalette& Palette, FToolElement& E
 					.ColorAndOpacity(Style.TitleForegroundColor)
 		]];
 	
+	TSharedPtr<SWidget> PaletteButtonsWidget = Element.GenerateWidget();
+	PaletteButtonsWidget->SetVisibility(TAttribute<EVisibility>::CreateSP(this, &FToolkitBuilder::GetActivePaletteCommandsVisibility));
+
 	ToolPaletteWidget->AddSlot()
 		.HAlign(HAlign_Fill)
 		.AutoHeight()
 		[
-			Element.GenerateWidget()
+			PaletteButtonsWidget.ToSharedRef()
 		];
 }
 
@@ -609,4 +612,9 @@ void FToolkitBuilder::SetCategoryButtonLabelVisibility(bool bIsCategoryButtonLab
 	SetCategoryButtonLabelVisibility(bIsCategoryButtonLabelVisible ? EVisibility::Visible : EVisibility::Collapsed);
 }
 
+void FToolkitBuilder::SetActivePaletteCommandsVisibility(EVisibility Visibility)
+{
+	ActivePaletteButtonVisibility = Visibility;
+}
+
 #undef LOCTEXT_NAMESPACE
\ No newline at end of file
diff --git a/Engine/Source/Developer/WidgetRegistration/Public/ToolkitBuilder.h b/Engine/Source/Developer/WidgetRegistration/Public/ToolkitBuilder.h
--- a/Engine/Source/Developer/WidgetRegistration/Public/ToolkitBuilder.h
+++ b/Engine/Source/Developer/WidgetRegistration/Public/ToolkitBuilder.h
@@ -176,6 +176,23 @@ public:
 	 */
 	void SetCategoryButtonLabelVisibility(bool bIsCategoryButtonLabelVisible);
 
+
+	/**
+	 * Configure whether or not the tool buttons in the active palette should be visible.
+	 * This can be dynamically toggled without rebuilding/regenerating the widget.
+	 * 
+	 * @param Visibility If Visibility == EVisibility::Collapsed, the category button labels
+	 * will be shown, else they will not be shown.
+	 */
+	void SetActivePaletteCommandsVisibility(EVisibility Visibility);
+
+
+	/**
+	 * @return the Visibility state of the active palette command buttons
+	 */
+	EVisibility GetActivePaletteCommandsVisibility() const { return ActivePaletteButtonVisibility; }
+
+
 	/**
 	 * RefreshCategoryToolbarWidget refreshes the UI display of the category toolbar 
 	 */
@@ -365,4 +382,8 @@ private:
 	/** If CategoryButtonLabelVisibility == EVisibility::Visible, the category button
 	 * labels are visible, else they are not displayed. By default the selected category button labels are Visible    */	
 	EVisibility CategoryButtonLabelVisibility;
+
+	/** If ActivePaletteButtonVisibility == EVisibility::Visible, the command buttons in the active palette
+	 * are visible, otherwise they are not displayed. By default the state is Visible  */
+	EVisibility ActivePaletteButtonVisibility = EVisibility::Visible;
 };
\ No newline at end of file
