ToolkitBuilder: add OnActivePaletteChanged delegate that will fire when active tool palette changes

diff --git a/Engine/Source/Developer/WidgetRegistration/Private/ToolkitBuilder.cpp b/Engine/Source/Developer/WidgetRegistration/Private/ToolkitBuilder.cpp
--- a/Engine/Source/Developer/WidgetRegistration/Private/ToolkitBuilder.cpp
+++ b/Engine/Source/Developer/WidgetRegistration/Private/ToolkitBuilder.cpp
@@ -321,6 +321,9 @@ void FToolkitBuilder::TogglePalette(TSharedPtr<FToolPalette> Palette)
 	}
 	
 	CreatePalette(Palette);
+
+	// active palette has changed
+	OnActivePaletteChanged.Broadcast();
 }
 
 void FToolkitBuilder::CreatePalette(TSharedPtr<FToolPalette> Palette)
diff --git a/Engine/Source/Developer/WidgetRegistration/Public/ToolkitBuilder.h b/Engine/Source/Developer/WidgetRegistration/Public/ToolkitBuilder.h
--- a/Engine/Source/Developer/WidgetRegistration/Public/ToolkitBuilder.h
+++ b/Engine/Source/Developer/WidgetRegistration/Public/ToolkitBuilder.h
@@ -67,7 +67,7 @@ struct WIDGETREGISTRATION_API FEditablePalette : public FToolPalette
 
 	/** Delegate used by FToolkitBuilder that is called when an item is added/removed from the palette */
 	FSimpleDelegate OnPaletteEdited;
-	
+
 	/**
 	 * Given a reference to a FUICommandInfo, returns whether it is in the current Palette
 	 *
@@ -255,6 +255,12 @@ public:
 	 */
 	ECheckBoxState IsActiveToolPalette(FName CommandName) const;
 
+	/** 
+	 * OnActivePaletteChanged is broadcast when the active palette changes to a different palette 
+	 */
+	FSimpleMulticastDelegate OnActivePaletteChanged;
+
+
 	void SetActiveToolDisplayName(FText InActiveToolDisplayName);
 
 	FText GetActiveToolDisplayName() const;
