ModelingMode: Move preset access UI into the active tool's header bar.

diff --git a/Engine/Source/Developer/WidgetRegistration/Private/ToolkitBuilder.cpp b/Engine/Source/Developer/WidgetRegistration/Private/ToolkitBuilder.cpp
--- a/Engine/Source/Developer/WidgetRegistration/Private/ToolkitBuilder.cpp
+++ b/Engine/Source/Developer/WidgetRegistration/Private/ToolkitBuilder.cpp
@@ -492,6 +492,9 @@ EVisibility FToolkitBuilder::GetActiveToolTitleVisibility() const
 void FToolkitBuilder::DefineWidget()
 {
 	RefreshCategoryToolbarWidget();
+
+	TSharedPtr<SHorizontalBox> ToolNameHeaderBox;
+
 	ToolkitWidgetVBox = SNew(SVerticalBox);
 	ToolkitWidgetHBox = 
 		SNew(SSplitter)
@@ -537,32 +540,33 @@ void FToolkitBuilder::DefineWidget()
 			.Padding(Style.ActiveToolTitleBorderPadding)
 			.BorderImage(&Style.ToolDetailsBackgroundBrush)
 			[
-			SNew(SBorder)
-			.Visibility(SharedThis(this), &FToolkitBuilder::GetActiveToolTitleVisibility)
-			.BorderImage(&Style.TitleBackgroundBrush)
-			.Padding(Style.ToolContextTextBlockPadding)
-			.HAlign(HAlign_Left)
-			.VAlign(VAlign_Top)
-			[
-
-				SNew(STextBlock)
+				SNew(SBorder)
+				.Visibility(SharedThis(this), &FToolkitBuilder::GetActiveToolTitleVisibility)
+				.BorderImage(&Style.TitleBackgroundBrush)
+				.Padding(Style.ToolContextTextBlockPadding)
+				[
+					SAssignNew(ToolNameHeaderBox, SHorizontalBox)
+					+SHorizontalBox::Slot()
+					.HAlign(EHorizontalAlignment::HAlign_Left)
+					[
+						SNew(STextBlock)
 						.Justification(ETextJustify::Left)
 						.Font(Style.TitleFont)
 						.Text(TAttribute<FText>(SharedThis(this), &FToolkitBuilder::GetActiveToolDisplayName))
 						.ColorAndOpacity(Style.TitleForegroundColor)
+					]
+				]
+
 			]
-		]
 	];
 	
 	if (ToolkitSections->ToolPresetArea)
 	{
-		ToolkitWidgetVBox->AddSlot()
-			.AutoHeight()
-			.HAlign(HAlign_Fill)
-			.Padding(5)
-			[
-				ToolkitSections->ToolPresetArea->AsShared()
-			];
+		ToolNameHeaderBox->AddSlot()
+		.HAlign(EHorizontalAlignment::HAlign_Right)
+		[
+			ToolkitSections->ToolPresetArea->AsShared()
+		];
 	}
 
 	if (ToolkitSections->ToolWarningArea)
