MeshModelingTools: Let clicking the currently active category re-show the tool buttons that are hidden when a tool is open.

diff --git a/Engine/Plugins/Editor/ModelingToolsEditorMode/Source/ModelingToolsEditorMode/Private/ModelingToolsEditorModeToolkit.cpp b/Engine/Plugins/Editor/ModelingToolsEditorMode/Source/ModelingToolsEditorMode/Private/ModelingToolsEditorModeToolkit.cpp
--- a/Engine/Plugins/Editor/ModelingToolsEditorMode/Source/ModelingToolsEditorMode/Private/ModelingToolsEditorModeToolkit.cpp
+++ b/Engine/Plugins/Editor/ModelingToolsEditorMode/Source/ModelingToolsEditorMode/Private/ModelingToolsEditorModeToolkit.cpp
@@ -227,6 +227,8 @@ void FModelingToolsEditorModeToolkit::RegisterPalettes()
 	ToolkitBuilderArgs.ToolkitCommandList = GetToolkitCommands();
 	ToolkitBuilderArgs.ToolkitSections = ToolkitSections;
 	ToolkitBuilderArgs.SelectedCategoryTitleVisibility = EVisibility::Collapsed;
+	// This lets us re-show the buttons if the user clicks a category with a tool still active.
+	ToolkitBuilderArgs.CategoryReclickBehavior = FToolkitBuilder::ECategoryReclickBehavior::TreatAsChanged;
 	ToolkitBuilder = MakeShared<FToolkitBuilder>(ToolkitBuilderArgs);
 	ToolkitBuilder->SetCategoryButtonLabelVisibility(UISettings->bShowCategoryButtonLabels);
 
diff --git a/Engine/Source/Developer/WidgetRegistration/Private/ToolkitBuilder.cpp b/Engine/Source/Developer/WidgetRegistration/Private/ToolkitBuilder.cpp
--- a/Engine/Source/Developer/WidgetRegistration/Private/ToolkitBuilder.cpp
+++ b/Engine/Source/Developer/WidgetRegistration/Private/ToolkitBuilder.cpp
@@ -92,7 +92,8 @@ FToolkitBuilder::FToolkitBuilder(const FToolkitBuilderArgs& Args) :
 	ToolbarCustomizationName(Args.ToolbarCustomizationName),
 	ToolkitCommandList(Args.ToolkitCommandList),
 	ToolkitSections(Args.ToolkitSections),
-	SelectedCategoryTitleVisibility(Args.SelectedCategoryTitleVisibility)
+	SelectedCategoryTitleVisibility(Args.SelectedCategoryTitleVisibility),
+	CategoryReclickBehavior(Args.CategoryReclickBehavior)
 {
 	SetCategoryButtonLabelVisibility(Args.bShowCategoryButtonLabels);
 	ResetWidget();
@@ -312,12 +313,20 @@ void FToolkitBuilder::TogglePalette(TSharedPtr<FToolPalette> Palette)
 	
 	if (ActivePalette && ActivePalette->LoadToolPaletteAction->GetCommandName() == CommandName)
 	{
-		return;
-/*
- *		@TODO: ~ enable this code for hiding category select on toggle		
- *		ActivePalette = nullptr;
-		ResetToolPaletteWidget();
-		return; */
+		switch (CategoryReclickBehavior)
+		{
+		case ECategoryReclickBehavior::NoEffect:
+			return;
+		case ECategoryReclickBehavior::ToggleOff:
+			ActivePalette = nullptr;
+			ResetToolPaletteWidget();
+			return;
+		case ECategoryReclickBehavior::TreatAsChanged:
+			// Fall through to the other handling
+			break;
+		default:
+			ensure(false);
+		}
 	}
 	
 	CreatePalette(Palette);
diff --git a/Engine/Source/Developer/WidgetRegistration/Public/ToolkitBuilder.h b/Engine/Source/Developer/WidgetRegistration/Public/ToolkitBuilder.h
--- a/Engine/Source/Developer/WidgetRegistration/Public/ToolkitBuilder.h
+++ b/Engine/Source/Developer/WidgetRegistration/Public/ToolkitBuilder.h
@@ -98,33 +98,7 @@ protected:
 	FGetEditableToolPaletteConfigManager GetConfigManager;
 };
 
-/*
- * A simple struct to carry initialization data for an FToolkitBuilder
- */
-struct FToolkitBuilderArgs
-{
-	FToolkitBuilderArgs(const FName& InToolbarCustomizationName) : ToolbarCustomizationName(InToolbarCustomizationName) 
-	{
-	}
-	
-	/** Name of the toolbar this mode uses and can be used by external systems to customize that mode toolbar */
-	const FName& ToolbarCustomizationName;
-
-	/** A TSharedPointer to the FUICommandList for the current mode */
-	TSharedPtr<FUICommandList> ToolkitCommandList;
-
-	/** The FToolkitSections which holds the sections defined for this Toolkit */
-	TSharedPtr<FToolkitSections> ToolkitSections;
-
-	/** If bShowCategoryButtonLabels == true, the category button
-	* labels should be, else they are not displayed  */	
-	bool bShowCategoryButtonLabels;
-
-	/** If SelectedCategoryTitleVisibility == EVisibility::Visible, the selected
-	 * category title is visible, else it is not displayed  */	
-	EVisibility SelectedCategoryTitleVisibility;
-};
-
+struct FToolkitBuilderArgs;
 
 /**
  * The FToolElementRegistrationArgs which is specified for Toolkits
@@ -141,6 +115,20 @@ public:
 	
 	virtual  ~FToolkitBuilder() override;
 	
+	// Used to specify what happens when you click the category button of a category that is already active.
+	enum class ECategoryReclickBehavior : uint8
+	{
+		// Do nothing if the same category button is clicked.
+		NoEffect,
+
+		// Toggle the active category off, so no category is active.
+		ToggleOff,
+
+		// Do the same thing that would be done if we were switching from a different category. Note that
+		// this will trigger OnActivePaletteChanged.
+		TreatAsChanged,
+	};
+
 	/*
 	 * Initializes the data to necessary to build the category toolbar
 	 *
@@ -392,4 +380,39 @@ private:
 	/** If ActivePaletteButtonVisibility == EVisibility::Visible, the command buttons in the active palette
 	 * are visible, otherwise they are not displayed. By default the state is Visible  */
 	EVisibility ActivePaletteButtonVisibility = EVisibility::Visible;
-};
\ No newline at end of file
+
+	/** Specifies what happens if you click the category button of an already-active catogory */
+	ECategoryReclickBehavior CategoryReclickBehavior = ECategoryReclickBehavior::NoEffect;
+};
+
+//~ This is defined after the class declaration so we can use the nested enum, which there isn't
+//~ a way to forward-declare.
+/*
+ * A simple struct to carry initialization data for an FToolkitBuilder
+ */
+struct FToolkitBuilderArgs
+{
+	FToolkitBuilderArgs(const FName& InToolbarCustomizationName) : ToolbarCustomizationName(InToolbarCustomizationName)
+	{
+	}
+
+	/** Name of the toolbar this mode uses and can be used by external systems to customize that mode toolbar */
+	const FName& ToolbarCustomizationName;
+
+	/** A TSharedPointer to the FUICommandList for the current mode */
+	TSharedPtr<FUICommandList> ToolkitCommandList;
+
+	/** The FToolkitSections which holds the sections defined for this Toolkit */
+	TSharedPtr<FToolkitSections> ToolkitSections;
+
+	/** If bShowCategoryButtonLabels == true, the category button
+	* labels should be, else they are not displayed  */
+	bool bShowCategoryButtonLabels;
+
+	/** If SelectedCategoryTitleVisibility == EVisibility::Visible, the selected
+	 * category title is visible, else it is not displayed  */
+	EVisibility SelectedCategoryTitleVisibility;
+
+	/** Specifies what happens if you click the category button of an already-active catogory */
+	FToolkitBuilder::ECategoryReclickBehavior CategoryReclickBehavior = FToolkitBuilder::ECategoryReclickBehavior::NoEffect;
+};
