Toolkitbuilder ~ don't show category picker if only one category, API call for active palette name

diff --git a/Engine/Source/Developer/WidgetRegistration/Private/ToolkitBuilder.cpp b/Engine/Source/Developer/WidgetRegistration/Private/ToolkitBuilder.cpp
--- a/Engine/Source/Developer/WidgetRegistration/Private/ToolkitBuilder.cpp
+++ b/Engine/Source/Developer/WidgetRegistration/Private/ToolkitBuilder.cpp
@@ -153,6 +153,15 @@ void FToolkitBuilder::GetCommandsForEditablePalette(TSharedRef<FEditablePalette>
 	}
 }
 
+FName FToolkitBuilder::GetActivePaletteName() const
+{
+	if (ActivePalette.IsValid() && ActivePalette->LoadToolPaletteAction.IsValid())
+	{
+		return ActivePalette->LoadToolPaletteAction->GetCommandName();
+	}
+	return NAME_None;
+}
+
 void FToolkitBuilder::AddPalette(TSharedPtr<FEditablePalette> Palette)
 {
 	Palette->OnPaletteEdited.BindSP(this, &FToolkitBuilder::OnEditablePaletteEdited, Palette.ToSharedRef());
@@ -224,6 +233,7 @@ void FToolkitBuilder::UpdateWidget()
 		UpdateEditablePalette(EditablePalette);		
 	}
 	
+	CategoryToolbarVisibility = LoadCommandNameToToolPaletteMap.Num() > 1 ? EVisibility::Visible : EVisibility::Collapsed;
 }
 
 void FToolkitBuilder::ToggleCommandInPalette(TSharedRef<FEditablePalette> Palette, FString CommandNameString)
@@ -275,7 +285,8 @@ void FToolkitBuilder::InitCategoryToolbarContainerWidget()
 {
 	if (!CategoryToolbarVBox.IsValid())
 	{
-		CategoryToolbarVBox = SNew(SVerticalBox);		
+		CategoryToolbarVBox = SNew(SVerticalBox)
+			.Visibility(CategoryToolbarVisibility);
 	}
 	else
 	{
@@ -476,11 +487,11 @@ void FToolkitBuilder::SetActivePaletteOnLoad(const FUICommandInfo* Command)
 
 TSharedPtr<SWidget> FToolkitBuilder::GenerateWidget()
 {
-	if (!ToolkitWidgetHBox)
+	if (!ToolkitWidgetContainerVBox)
 	{
 		DefineWidget();
 	}
-	return ToolkitWidgetHBox.ToSharedRef();
+	return ToolkitWidgetContainerVBox.ToSharedRef();
 }
 
 void FToolkitBuilder::SetActiveToolDisplayName(FText InActiveToolDisplayName)
@@ -504,8 +515,10 @@ void FToolkitBuilder::DefineWidget()
 
 	TSharedPtr<SHorizontalBox> ToolNameHeaderBox;
 
+	ToolkitWidgetContainerVBox = SNew(SVerticalBox);
+
 	ToolkitWidgetVBox = SNew(SVerticalBox);
-	ToolkitWidgetHBox = 
+	TSharedPtr<SWidget> MainSplitter = 
 		SNew(SSplitter)
 		+ SSplitter::Slot()
 		.Resizable(false)
@@ -519,6 +532,13 @@ void FToolkitBuilder::DefineWidget()
 			[
 				ToolkitWidgetVBox->AsShared()
 			];
+	
+	ToolkitWidgetContainerVBox->AddSlot()
+	.VAlign(VAlign_Fill)
+	.FillHeight(1)
+	[
+		MainSplitter->AsShared()
+	];
 
 	if (ToolkitSections->ModeWarningArea)
 	{
@@ -556,10 +576,13 @@ void FToolkitBuilder::DefineWidget()
 				[
 					SAssignNew(ToolNameHeaderBox, SHorizontalBox)
 					+SHorizontalBox::Slot()
+					.Padding(0)
+					.VAlign(EVerticalAlignment::VAlign_Center)
 					.HAlign(EHorizontalAlignment::HAlign_Left)
 					[
 						SNew(STextBlock)
 						.Justification(ETextJustify::Left)
+						.Margin(0)
 						.Font(Style.TitleFont)
 						.Text(TAttribute<FText>(SharedThis(this), &FToolkitBuilder::GetActiveToolDisplayName))
 						.ColorAndOpacity(Style.TitleForegroundColor)
diff --git a/Engine/Source/Developer/WidgetRegistration/Public/ToolkitBuilder.h b/Engine/Source/Developer/WidgetRegistration/Public/ToolkitBuilder.h
--- a/Engine/Source/Developer/WidgetRegistration/Public/ToolkitBuilder.h
+++ b/Engine/Source/Developer/WidgetRegistration/Public/ToolkitBuilder.h
@@ -147,7 +147,7 @@ public:
 	void InitCategoryToolbarContainerWidget();
 
 	/**
-	 * Sets category button label visibility to Visiblity. It also reinitializes the
+	 * Sets category button label visibility to Visibility. It also reinitializes the
 	 * category toolbar data, as the toolbar's LabelVisibility member is now stale.
 	 *
 	 * @param Visibility If Visibility == EVisibility::Collapsed, the category button labels
@@ -174,13 +174,11 @@ public:
 	 */
 	void SetActivePaletteCommandsVisibility(EVisibility Visibility);
 
-
 	/**
 	 * @return the Visibility state of the active palette command buttons
 	 */
 	EVisibility GetActivePaletteCommandsVisibility() const { return ActivePaletteButtonVisibility; }
 
-
 	/**
 	 * RefreshCategoryToolbarWidget refreshes the UI display of the category toolbar 
 	 */
@@ -248,12 +246,30 @@ public:
 	 */
 	FSimpleMulticastDelegate OnActivePaletteChanged;
 
-
+	/**
+	 * Sets the display name for the active tool
+	 *
+	 * @param InActiveToolDisplayName an FText that holds the name of the currently active tool
+	 */
 	void SetActiveToolDisplayName(FText InActiveToolDisplayName);
 
+	/**
+	 * @returns the display name for the active tool
+	 */
 	FText GetActiveToolDisplayName() const;
 
+	/**
+	 * Fills the OutCommands TArray with TSharedPtr<const FUICommandInfo> instances that are present in the FEditablePalette
+	 * EditablePalette
+	 *
+	 * @param EditablePalette the FEditablePalette whose FUICommandInfo
+	 */
 	void GetCommandsForEditablePalette(TSharedRef<FEditablePalette> EditablePalette, TArray<TSharedPtr<const FUICommandInfo>>& OutCommands);
+
+	/**
+	 * @returns the name of the Active Palette if one is available, else it returns NAME_NONE
+	 */
+	FName GetActivePaletteName() const;
 	
 private:
 
@@ -357,8 +373,8 @@ private:
 	/** The SVerticalBox which holds all but the vertical toolbar in a Toolkit */
 	TSharedPtr<SVerticalBox> ToolkitWidgetVBox;
 
-	/** The SSplitter which holds the entire Toolkit */
-	TSharedPtr<SSplitter> ToolkitWidgetHBox;
+	/** The SVerticalBox which holds the entire Toolkit */
+	TSharedPtr<SVerticalBox> ToolkitWidgetContainerVBox;
 
 	/** The FToolkitSections which holds the sections defined for this Toolkit */
 	TSharedPtr<FToolkitSections> ToolkitSections;
@@ -381,6 +397,12 @@ private:
 	 * are visible, otherwise they are not displayed. By default the state is Visible  */
 	EVisibility ActivePaletteButtonVisibility = EVisibility::Visible;
 
+	/*
+	 * Is the Category toolbar visible? This should be false if the Category toolbar has less than 2 categories,
+	 * because if there is only one showing it is superfluous to actually show it since it will always be selected.
+	 */
+	EVisibility CategoryToolbarVisibility;
+
 	/** Specifies what happens if you click the category button of an already-active catogory */
 	ECategoryReclickBehavior CategoryReclickBehavior = ECategoryReclickBehavior::NoEffect;
 };
