* UE Toolkit Builder fit and finish/cleanup.
* Added a Separator Builder to the Modular class display infrastructure.

diff --git a/Engine/Source/Developer/WidgetRegistration/Private/Layout/SeparatorBuilder.cpp b/Engine/Source/Developer/WidgetRegistration/Private/Layout/SeparatorBuilder.cpp
new file mode 100644
--- /dev/null
+++ b/Engine/Source/Developer/WidgetRegistration/Private/Layout/SeparatorBuilder.cpp
@@ -0,0 +1,46 @@
+// Copyright Epic Games, Inc. All Rights Reserved.
+
+#include "Layout/SeparatorBuilder.h"
+
+TSharedPtr<SWidget> FSeparatorBuilder::GenerateWidget()
+{
+	return SNew( SSeparator )
+				.SeparatorImage(Image)
+				.Thickness(Size)
+				.Orientation(Orientation)
+				.Visibility(Visibility)
+				.ColorAndOpacity(SlateColor);
+}
+
+TSharedRef<SSeparator> FSeparatorBuilder::ToSSeparatorSharedRef()
+{
+	return StaticCastSharedRef<SSeparator>(GenerateWidget().ToSharedRef());
+}
+
+TSharedRef<SSeparator> FSeparatorBuilder::operator*()
+{
+	return ToSSeparatorSharedRef();
+}
+
+FSeparatorBuilder& FSeparatorBuilder::InitializeSize(const FSeparatorSize& NewSize)
+{
+	Size = NewSize.ToFloat();
+	return *this;
+}
+
+FSeparatorBuilder& FSeparatorBuilder::SetColor(const EStyleColor& NewColor)
+{
+	// the overloaded conversion is needed here, as we will want the conversion from EStyleColor to FSlateColor
+	SlateColor = NewColor;
+	return *this;
+}
+
+FSeparatorBuilder::FSeparatorBuilder(EStyleColor InColor, EOrientation InOrientation, FSeparatorSize InSize,
+	const FSlateBrush* InImage):
+	FToolElementRegistrationArgs(EToolElement::Separator),
+	Image(InImage),
+	Orientation(InOrientation),
+	SlateColor(InColor),
+	Size(InSize.ToFloat())
+{
+}
diff --git a/Engine/Source/Developer/WidgetRegistration/Private/Layout/SeparatorTemplates.cpp b/Engine/Source/Developer/WidgetRegistration/Private/Layout/SeparatorTemplates.cpp
new file mode 100644
--- /dev/null
+++ b/Engine/Source/Developer/WidgetRegistration/Private/Layout/SeparatorTemplates.cpp
@@ -0,0 +1,18 @@
+// Copyright Epic Games, Inc. All Rights Reserved.
+
+#include "Layout/SeparatorTemplates.h"
+
+FSeparatorBuilder FSeparatorTemplates::SmallVerticalBackgroundNoBorder()
+{
+	return FSeparatorBuilder{ EStyleColor::Background, Orient_Vertical };
+}
+
+FSeparatorBuilder FSeparatorTemplates::SmallHorizontalBackgroundNoBorder()
+{
+	return FSeparatorBuilder{};
+}
+
+FSeparatorBuilder FSeparatorTemplates::SmallHorizontalPanelNoBorder()
+{
+	return FSeparatorBuilder{ EStyleColor::Panel };
+}
diff --git a/Engine/Source/Developer/WidgetRegistration/Private/Styles/SlateBrushTemplates.cpp b/Engine/Source/Developer/WidgetRegistration/Private/Styles/SlateBrushTemplates.cpp
new file mode 100644
--- /dev/null
+++ b/Engine/Source/Developer/WidgetRegistration/Private/Styles/SlateBrushTemplates.cpp
@@ -0,0 +1,8 @@
+// Copyright Epic Games, Inc. All Rights Reserved.
+
+#include "Styles/SlateBrushTemplates.h"
+
+const FSlateBrush* FSlateBrushTemplates::ThinLineHorizontal()
+{
+	return FAppStyle::GetBrush( "ThinLine.Horizontal");		
+}
diff --git a/Engine/Source/Developer/WidgetRegistration/Private/ToolkitBuilder.cpp b/Engine/Source/Developer/WidgetRegistration/Private/ToolkitBuilder.cpp
--- a/Engine/Source/Developer/WidgetRegistration/Private/ToolkitBuilder.cpp
+++ b/Engine/Source/Developer/WidgetRegistration/Private/ToolkitBuilder.cpp
@@ -9,6 +9,8 @@
 #include "ToolElementRegistry.h"
 #include "ToolkitStyle.h"
 #include "Framework/Commands/UICommandList.h"
+#include "Layout/SeparatorBuilder.h"
+#include "Layout/SeparatorTemplates.h"
 #include "Styling/StyleColors.h"
 
 #define LOCTEXT_NAMESPACE "ToolkitBuilder"
@@ -428,7 +430,6 @@ void FToolkitBuilder::CreatePaletteWidget(FToolPalette& Palette, FToolElement& E
 		];
 }
 
-
 TSharedRef<SWidget> FToolkitBuilder::GetToolPaletteWidget() const
 {
 	return ToolPaletteWidget->AsShared();
@@ -515,11 +516,14 @@ void FToolkitBuilder::DefineWidget()
 
 	TSharedPtr<SHorizontalBox> ToolNameHeaderBox;
 
-	ToolkitWidgetContainerVBox = SNew(SVerticalBox);
+	ToolkitWidgetContainerVBox = SNew(SVerticalBox)
+	+ SVerticalBox::Slot().AutoHeight() [ *FSeparatorTemplates::SmallHorizontalPanelNoBorder()  ]
+	+ SVerticalBox::Slot().AutoHeight() [ *FSeparatorTemplates::SmallHorizontalBackgroundNoBorder() ];
 
 	ToolkitWidgetVBox = SNew(SVerticalBox);
 	TSharedPtr<SWidget> MainSplitter = 
 		SNew(SSplitter)
+		.PhysicalSplitterHandleSize(2.0f)
 		+ SSplitter::Slot()
 		.Resizable(false)
 		.SizeRule(SSplitter::SizeToContent)
@@ -559,6 +563,11 @@ void FToolkitBuilder::DefineWidget()
 			GetToolPaletteWidget()
 		];
 
+	ToolkitWidgetVBox->AddSlot()
+		.AutoHeight()
+		[ *FSeparatorTemplates::SmallHorizontalBackgroundNoBorder().BindVisibility(
+			TAttribute<EVisibility>::CreateLambda([this] () { return ActivePaletteButtonVisibility; }))];
+	
 	ToolkitWidgetVBox->AddSlot()
 		.AutoHeight()
 		.HAlign(HAlign_Fill)
@@ -566,7 +575,7 @@ void FToolkitBuilder::DefineWidget()
 		[
 			SNew(SBorder)
 			.HAlign(HAlign_Fill)
-			.Padding(Style.ActiveToolTitleBorderPadding)
+			.Padding(Style.ActiveToolTitleBorderPadding)                                                                                                                                                                                                                                                                                                                                                                                                                                                                                                                                                                                                                                                                                                                                                                                                                                                                                                                                                                                                                                                                                                                                                                                                                                                                                                                                                                                                                                                                                                                                                                                                                                                                                                                                                                                                                                                                                                                                                                                                                                                                                                                                                                                                                                                                                                                                                                                                                                                                                                                                                                                                                                                                                                                                                                                                                                                                                                                                                                                                                                                                                                                                                                                                                                                                                                                                                                                                                                                                                                                                                                                                                                                                                                                                                                                                                                                                                                                                                                                                                                                                                                                                                                                                                                                                                                                                                                                                                                                                                                                                                                                                                                                                                                                                                                                                                                                                                                                                                                                                                                                                                                                                                                                                                                                                                                                                                                                                                                                                                                                                                                                                                                                                                                                                                                                                                                                                                                                                                                                                                                                                                                                                                                                                                                                                         
 			.BorderImage(&Style.ToolDetailsBackgroundBrush)
 			[
 				SNew(SBorder)
@@ -620,7 +629,7 @@ void FToolkitBuilder::DefineWidget()
 			[
 			SNew(SBorder)
 			.BorderImage(&Style.ToolDetailsBackgroundBrush)
-			.Padding(8.f, 2.f, 0.f, 2.f)
+			.Padding(0.f, 2.f, 0.f, 2.f)
 				[
 					ToolkitSections->DetailsView->AsShared()
 				]
diff --git a/Engine/Source/Developer/WidgetRegistration/Private/ToolkitStyle.cpp b/Engine/Source/Developer/WidgetRegistration/Private/ToolkitStyle.cpp
--- a/Engine/Source/Developer/WidgetRegistration/Private/ToolkitStyle.cpp
+++ b/Engine/Source/Developer/WidgetRegistration/Private/ToolkitStyle.cpp
@@ -29,12 +29,12 @@ void FToolkitStyle::Initialize()
 	
 	FToolkitWidgetStyle Style;
 
-	Style.SetTitleBackgroundBrush(FSlateRoundedBoxBrush(FStyleColors::Header, 0.f, FLinearColor(0, 0, 0, .8), 0.5))
-		 .SetToolDetailsBackgroundBrush( FSlateRoundedBoxBrush(FStyleColors::Background, 0.f, FLinearColor(0, 0, 0, .8), 0.5))
+	Style.SetTitleBackgroundBrush(FSlateRoundedBoxBrush(FStyleColors::Header, 0.f, FLinearColor(0, 0, 0, .8), 0))
+		 .SetToolDetailsBackgroundBrush( FSlateRoundedBoxBrush(FStyleColors::Background, 0.f, FLinearColor(0, 0, 0, .8), 0))
 		 .SetTitleForegroundColor(FStyleColors::Foreground)
 		 .SetTitlePadding(FMargin(16.f, 4.f))
 	 	 .SetTitleFont(FCoreStyle::GetDefaultFontStyle("Bold", 10))
-		 .SetActiveToolTitleBorderPadding(FMargin(0.f, 8.f, 0.f, 0.f))
+		 .SetActiveToolTitleBorderPadding(FMargin(0.f, 0.f, 0.f, 0.f))
 		 .SetToolContextTextBlockPadding(FMargin(16.f, 4.f));
 
 	StyleSet->Set("FToolkitWidgetStyle", Style);
@@ -64,5 +64,6 @@ const ISlateStyle& FToolkitStyle::Get()
 	return *StyleSet;
 }
 
+
 #undef LOCAL_IMAGE_BRUSH
 #undef IMAGE_PLUGIN_BRUSH
diff --git a/Engine/Source/Developer/WidgetRegistration/Public/Layout/SeparatorBuilder.h b/Engine/Source/Developer/WidgetRegistration/Public/Layout/SeparatorBuilder.h
new file mode 100644
--- /dev/null
+++ b/Engine/Source/Developer/WidgetRegistration/Public/Layout/SeparatorBuilder.h
@@ -0,0 +1,165 @@
+// Copyright Epic Games, Inc. All Rights Reserved.
+
+#pragma once
+#include "ToolElementRegistry.h"
+#include "Styling/SlateTypes.h"
+#include "Styles/SlateBrushTemplates.h"
+#include "Styling/StyleColors.h"
+#include "Widgets/Layout/SSeparator.h"
+
+/**
+ * Provides a type that will hold the size for a FSeparatorBuilder
+ */
+class FSeparatorSize
+{
+	/**
+	 * The magnitude of the size ~ the greater the Magnitude the larger the Separator.
+	 */
+	float Magnitude;
+
+	/**
+	 * Constructs the Separator Size with a Magnitude o InMagnitude
+	 *
+	 * @param InMagnitude 
+	 */
+	FSeparatorSize(float InMagnitude): Magnitude(InMagnitude)
+	{
+	}
+	
+public:
+
+	/**
+	 * Converts the Size to a float value
+	 */
+	float ToFloat() const
+	{
+		return Magnitude;
+	}
+
+	friend struct FSeparatorSizes;
+};
+
+/**
+ * A struct to contain the possible sizes for a Separator 
+ */
+struct FSeparatorSizes
+{
+	/** the size of an extra small (1 slate unit) Separator */
+	static inline const FSeparatorSize XSmall{1.0f};
+	
+	/** the size of a small (2 slate unit) Separator */
+	static inline const FSeparatorSize Small{2.0f};
+	
+	/** the size of a medium (4 slate unit) Separator */
+	static inline const FSeparatorSize Medium{4.0f};
+	
+	/** the size of a large (8 slate unit) Separator */
+	static inline const FSeparatorSize Large{8.0f};
+	
+	/** the size of an extra large (16 slate unit) Separator */
+	static inline const FSeparatorSize XLarge{16.0f};
+};
+
+/**
+ * FSeparatorBuilder instances will build a TSharedRef<SSeparator>
+ */
+class FSeparatorBuilder final : public FToolElementRegistrationArgs
+{
+public:
+
+	/**
+	 * Implements the generation of the TSharedPtr<SWidget>
+	 */
+	virtual TSharedPtr<SWidget> GenerateWidget() override;
+
+	/**
+	 * Generates the SSeparator and returns the TSharedRef<SSeparator> as a reference to it
+	 *
+	 * @return the TSharedRef<SSeparator> as a reference to the generated SSeparator
+	 */
+	TSharedRef<SSeparator> ToSSeparatorSharedRef();
+
+	/**
+	 * Generates the SSeparator and returns the TSharedRef<SSeparator> as a reference to it.
+	 * This provides some convenience syntax to generate the widget and get a reference to it
+	 * in a concise way.
+	 *
+ 	 * @return the TSharedRef<SSeparator> as a reference to the generated SSeparator
+	 */
+	TSharedRef<SSeparator> operator*();
+
+	/**
+	 * Initializes the size of the separator. If called after the widget creation (that is, a call to GenerateWidget)
+	 * this call will have no effect.
+	 *
+	 * @param NewSize the FSeparatorSize& that will provide the size for the Separator
+	 * @return this FSeparatorBuilder to enable cascading.
+	 */
+	FSeparatorBuilder& InitializeSize(const FSeparatorSize& NewSize);
+
+	/** 
+	 * Sets the color of the Separator. This will update the color of the Separator even after widget creation.
+	 *
+	 * @param NewColor the EStyleColor to set the color of the SSeparator. Note that the EStyleColor is a key into
+	 * the color table, so providing an EStyleColor such as EStyleColor::Panel will honor updates to changes in themes
+	 * for that EStyleColor immediately.
+	 * @return this FSeparatorBuilder to enable cascading.
+	 */
+	FSeparatorBuilder& SetColor(const EStyleColor& NewColor);
+
+	/** 
+	 * Binds the Visibility of the Separator to a TAttribute for continuous update.
+	 *
+	 * @param NewVisibility the TAttribute<EVisibility> to bind for the SSeparator. 
+	 * @return this FSeparatorBuilder to enable cascading.
+	 */
+	FSeparatorBuilder& BindVisibility(TAttribute<EVisibility> NewVisibility)
+	{
+		Visibility = NewVisibility;
+		return *this;
+	}
+	
+private:
+
+	friend struct FSeparatorTemplates;
+
+	/**
+	 * Constructs the FSeparatorBuilder. This is private to keep the end user API simple. By starting with a template
+	 * callers should not have a need to construct an FSeparatorBuilder.
+	 *
+	 * @param InColor the EStyleColor to provide the color for this Separator
+	 * @param InOrientation the EOrientation that sets the orientation for this Separator
+	 * @param InSize the FSeparatorSize that sets the size for this Separator
+	 * @param InImage the const FSlateBrush* which specifies the FSlateBrush to use for this Separator 
+	 */
+	FSeparatorBuilder(EStyleColor InColor = EStyleColor::Background,
+		              EOrientation InOrientation = Orient_Horizontal,
+					  FSeparatorSize InSize = FSeparatorSizes::Small,
+					  const FSlateBrush* InImage = FSlateBrushTemplates::ThinLineHorizontal());
+
+	/**
+	 * The const FSlateBrush* for the SSeparator
+	 */
+	const FSlateBrush* Image =  FSlateBrushTemplates::ThinLineHorizontal();
+
+	/**
+	 * The EOrientation that specifies whether this is a horizontal or vertical Separator
+	 */
+	EOrientation Orientation = Orient_Horizontal;
+
+	/**
+	 * The FSlateColor that specifies the Color for this Separator. 
+	 */
+	TAttribute<FSlateColor> SlateColor = EStyleColor::Background;
+
+	/**
+	 * The FSeparatorSize that specifies the Size for this Separator.
+	 */
+	float Size = FSeparatorSizes::Small.ToFloat();
+
+	/**
+	 * The TAttribute<EVisibility> which will provide the Visibility for the Separator
+	 */
+	TAttribute<EVisibility> Visibility = EVisibility::Visible;
+
+};
\ No newline at end of file
diff --git a/Engine/Source/Developer/WidgetRegistration/Public/Layout/SeparatorTemplates.h b/Engine/Source/Developer/WidgetRegistration/Public/Layout/SeparatorTemplates.h
new file mode 100644
--- /dev/null
+++ b/Engine/Source/Developer/WidgetRegistration/Public/Layout/SeparatorTemplates.h
@@ -0,0 +1,39 @@
+// Copyright Epic Games, Inc. All Rights Reserved.
+
+#pragma once
+#include "SeparatorBuilder.h"
+
+
+/**
+ * A struct which holds methods that provide templates for FSeparatorBuilders. The provided FSeparatorBuilders
+ * can be updated as needed with the FSeparatorBuilder API.
+ */
+struct FSeparatorTemplates
+{
+	/**
+	* Creates and returns a small (2 Slate Unit), vertical FSeparatorBuilder with the color of the theme
+	 * specified color for Background
+	 *
+	 * @return a small (2 Slate Unit), vertical FSeparatorBuilder with the color of the theme
+	 * specified color for Background
+	 */
+	static FSeparatorBuilder SmallVerticalBackgroundNoBorder();
+
+	/**
+	* Creates and returns a small (2 Slate Unit), horizontal FSeparatorBuilder set to the color of the theme
+	* specified color for Background
+	*
+	* @return a small (2 Slate Unit), horizontal FSeparatorBuilder with the color of the theme
+	* specified color for Background
+	 */
+	static FSeparatorBuilder SmallHorizontalBackgroundNoBorder();
+
+	/**
+	* Creates and returns a small (2 Slate Unit), horizontal FSeparatorBuilder set to the color of the theme
+	 * specified color for Panels.
+	 *
+	 * @return a small (2 Slate Unit), horizontal FSeparatorBuilder with the color of the theme
+	 * specified color for Panels.
+	 */
+	static FSeparatorBuilder SmallHorizontalPanelNoBorder();
+};
diff --git a/Engine/Source/Developer/WidgetRegistration/Public/Styles/SlateBrushTemplates.h b/Engine/Source/Developer/WidgetRegistration/Public/Styles/SlateBrushTemplates.h
new file mode 100644
--- /dev/null
+++ b/Engine/Source/Developer/WidgetRegistration/Public/Styles/SlateBrushTemplates.h
@@ -0,0 +1,11 @@
+// Copyright Epic Games, Inc. All Rights Reserved.
+
+#pragma once
+#include "Styling/AppStyle.h"
+#include "Styling/SlateBrush.h"
+
+/** A const FSlateBrush* Factory */
+struct WIDGETREGISTRATION_API FSlateBrushTemplates
+{
+	static const FSlateBrush* ThinLineHorizontal();
+};
diff --git a/Engine/Source/Developer/WidgetRegistration/Public/ToolElementRegistry.h b/Engine/Source/Developer/WidgetRegistration/Public/ToolElementRegistry.h
--- a/Engine/Source/Developer/WidgetRegistration/Public/ToolElementRegistry.h
+++ b/Engine/Source/Developer/WidgetRegistration/Public/ToolElementRegistry.h
@@ -10,7 +10,8 @@ class FToolElementRegistry;
 enum class EToolElement : uint32
 {
 	Toolbar = 0,
-	Toolkit
+	Toolkit,
+	Separator
 };
 
 /** A class which provides the data necessary to generate a widget. This class should be
