Place Actors Panel - Usability Improvements

~ Extracted a large chunk out of FToolkitBuilder and put it in abstract class CategoryDrivenContentBuilderBase, which provides a base class for any class which needs a lefthand vertical category picker which loads content in the box next to the picker
~ Implemented CategoryDrivenContentBuilder, which provides a delegate which takes an FName as an argument which provides the name of the category just clicked, and returns the TSharedRef<SWidget> the user wants loaded for the selected category
~ created FPlacementModeCommands which creates a command to load content for any placement commands found
~ Introduced CVar PlaceActors.EnableCategoryChooserView which, when true, loads the new (in progress) place actors tab
~ altered SPlacementModeTools to create and use the CategoryDrivenContentBuilder when PlaceActors.EnableCategoryChooserView is true, else it loads the usual place actors panel
~ updated ToolElementRegistry with a convenience method that generates the widget and immediately converts it to a shared Ref

Note that after a change in the CVar, you will need to close and reopen the place actors panel for it to appear properly.

diff --git a/Engine/Source/Developer/WidgetRegistration/Private/Layout/CategoryDrivenContentBuilder.cpp b/Engine/Source/Developer/WidgetRegistration/Private/Layout/CategoryDrivenContentBuilder.cpp
new file mode 100644
--- /dev/null
+++ b/Engine/Source/Developer/WidgetRegistration/Private/Layout/CategoryDrivenContentBuilder.cpp
@@ -0,0 +1,67 @@
+// Copyright Epic Games, Inc. All Rights Reserved.
+
+
+#include "Layout/CategoryDrivenContentBuilder.h"
+#include "ToolkitBuilder.h"
+#include "Templates/SharedPointer.h"
+
+FCategoryDrivenContentBuilder ::FCategoryDrivenContentBuilder(FCategoryDrivenContentBuilderArgs& Args):
+	FCategoryDrivenContentBuilderBase( Args  )
+{
+}
+
+FCategoryDrivenContentBuilder::~FCategoryDrivenContentBuilder()
+{
+	ProvideSelectedCategoryContentDelegate.Unbind();
+}
+
+void FCategoryDrivenContentBuilder::ProvideSelectedCategoryContent( FName InActiveCommandName )
+{
+	ActiveCommandName = InActiveCommandName;
+	UpdateWidget();
+}
+
+
+void FCategoryDrivenContentBuilder::InitializeCategoryToolbar()
+{
+	for ( const TSharedPtr<FUICommandInfo>& Command : ContentLoaderCommands )
+	{
+		LoadToolPaletteCommandList->MapAction(
+			Command,
+			FExecuteAction::CreateSP(this, &FCategoryDrivenContentBuilder::ProvideSelectedCategoryContent, Command->GetCommandName()),
+			FCanExecuteAction(),
+			FGetActionCheckState::CreateSP(this, &FCategoryDrivenContentBuilder::IsActiveToolPalette,
+			                               Command->GetCommandName())
+		);
+		LoadPaletteToolBarBuilder->AddToolBarButton(Command);
+	}
+}
+
+void FCategoryDrivenContentBuilder::UpdateWidget()
+{
+	ToolkitWidgetVBox->ClearChildren();
+
+	if ( ProvideSelectedCategoryContentDelegate.IsBound() )
+	{
+		TSharedRef<SWidget> Widget =  ProvideSelectedCategoryContentDelegate.Execute( ActiveCommandName );
+		ToolkitWidgetVBox->AddSlot()
+		.FillHeight(1.0f)
+		.VAlign(VAlign_Fill)
+		 [
+		   Widget
+		 ];
+	}
+}
+
+void FCategoryDrivenContentBuilder::SetCommands(TArray<TSharedPtr<FUICommandInfo>> InContentLoaderCommands)
+{
+	ContentLoaderCommands = InContentLoaderCommands;
+	const bool bForceSmallIcons = true; 
+	
+	Style = FToolkitStyle::Get().GetWidgetStyle<FToolkitWidgetStyle>("FToolkitWidgetStyle");
+	LoadToolPaletteCommandList = MakeShared<FUICommandList>();
+	LoadPaletteToolBarBuilder = MakeShared<FVerticalToolBarBuilder>(LoadToolPaletteCommandList, FMultiBoxCustomization::None, TSharedPtr<FExtender>(), bForceSmallIcons);
+	LoadPaletteToolBarBuilder->SetLabelVisibility( CategoryButtonLabelVisibility );
+	ToolkitWidgetVBox = SNew(SVerticalBox);
+	InitializeCategoryToolbar();
+}
diff --git a/Engine/Source/Developer/WidgetRegistration/Private/Layout/CategoryDrivenContentBuilderBase.cpp b/Engine/Source/Developer/WidgetRegistration/Private/Layout/CategoryDrivenContentBuilderBase.cpp
new file mode 100644
--- /dev/null
+++ b/Engine/Source/Developer/WidgetRegistration/Private/Layout/CategoryDrivenContentBuilderBase.cpp
@@ -0,0 +1,162 @@
+// Copyright Epic Games, Inc. All Rights Reserved.
+
+#include "Layout/CategoryDrivenContentBuilderBase.h"
+
+#include "ToolbarRegistrationArgs.h"
+
+
+#include "Widgets/SBoxPanel.h"
+#include "Widgets/Layout/SSplitter.h"
+#include "Widgets/SBoxPanel.h"
+#include "ToolElementRegistry.h"
+#include "Framework/Commands/UICommandList.h"
+#include "Layout/SeparatorBuilder.h"
+#include "Layout/SeparatorTemplates.h"
+#include "Layout/Visibility.h"
+
+namespace UE::CategoryDrivenContentBuilderBase::Private
+{
+	const FName StyleName = TEXT("Name");
+}
+
+FToolElementRegistry FCategoryDrivenContentBuilderBase::ToolRegistry = FToolElementRegistry::Get();
+
+FCategoryDrivenContentBuilderBase::FCategoryDrivenContentBuilderBase
+	(
+		FName InBuilderName
+	) :
+	FToolElementRegistrationArgs( UE::CategoryDrivenContentBuilderBase::Private::StyleName ),
+        CategoryReclickBehavior( ECategoryReclickBehavior::NoEffect) ,
+	BuilderName( InBuilderName )
+{
+}
+
+FCategoryDrivenContentBuilderBase::FCategoryDrivenContentBuilderBase
+	(
+		FCategoryDrivenContentBuilderArgs& Args
+	) :
+	FToolElementRegistrationArgs(Args.BuilderName ),
+	CategoryReclickBehavior( Args.CategoryReclickBehavior ),
+	BuilderName(Args.BuilderName)
+{
+}
+
+FCategoryDrivenContentBuilderBase::~FCategoryDrivenContentBuilderBase()
+{
+	if (VerticalToolbarElement.IsValid())
+	{
+		ToolRegistry.UnregisterElement(VerticalToolbarElement.ToSharedRef());
+	}
+}
+
+TSharedPtr<FToolBarBuilder> FCategoryDrivenContentBuilderBase::GetLoadPaletteToolbar()
+{
+	return LoadPaletteToolBarBuilder;
+}
+
+void FCategoryDrivenContentBuilderBase::InitCategoryToolbarContainerWidget()
+{
+	if (!CategoryToolbarVBox.IsValid())
+	{
+		CategoryToolbarVBox = SNew(SVerticalBox)
+			.Visibility(CategoryToolbarVisibility);
+	}
+	else
+	{
+		CategoryToolbarVBox->ClearChildren();
+	}
+	CategoryToolbarVBox->AddSlot()
+	.Padding(0.f)
+	[
+		CreateToolbarWidget()
+	];
+}
+
+void FCategoryDrivenContentBuilderBase::RefreshCategoryToolbarWidget()
+{
+	FToolElementRegistrationKey Key = FToolElementRegistrationKey(BuilderName, EToolElement::Toolbar);
+	VerticalToolbarElement = ToolRegistry.GetToolElementSP(Key);
+	const TSharedRef<FToolbarRegistrationArgs> VerticalToolbarRegistrationArgs = MakeShareable<FToolbarRegistrationArgs>(
+		new FToolbarRegistrationArgs(LoadPaletteToolBarBuilder.ToSharedRef()));
+	
+	if (!VerticalToolbarElement.IsValid())
+	{
+		VerticalToolbarElement = MakeShareable(new FToolElement
+			(BuilderName,
+			VerticalToolbarRegistrationArgs));
+		ToolRegistry.RegisterElement(VerticalToolbarElement.ToSharedRef());
+	}
+
+	VerticalToolbarElement->SetRegistrationArgs(VerticalToolbarRegistrationArgs);
+	InitCategoryToolbarContainerWidget();
+}
+
+TSharedPtr<SWidget> FCategoryDrivenContentBuilderBase::GenerateWidget()
+{
+	if (!ToolkitWidgetContainerVBox)
+	{
+		CreateWidget();
+	}
+	return ToolkitWidgetContainerVBox.ToSharedRef();
+}
+
+
+void FCategoryDrivenContentBuilderBase::CreateWidget()
+{
+	ToolkitWidgetVBox = ToolkitWidgetVBox.IsValid() ? ToolkitWidgetVBox : SNew(SVerticalBox);
+	ToolkitWidgetVBox->ClearChildren();
+	RefreshCategoryToolbarWidget();
+	ProvideSelectedCategoryContent();
+
+	ToolkitWidgetContainerVBox = SNew(SVerticalBox)
+	+ SVerticalBox::Slot().AutoHeight() [ *FSeparatorTemplates::SmallHorizontalPanelNoBorder()  ]
+	+ SVerticalBox::Slot().AutoHeight() [ *FSeparatorTemplates::SmallHorizontalBackgroundNoBorder() ];
+
+	TSharedPtr<SWidget> MainSplitter = 
+		SNew(SSplitter)
+		.PhysicalSplitterHandleSize(2.0f)
+		+ SSplitter::Slot()
+		.Resizable(false)
+		.SizeRule(SSplitter::SizeToContent)
+			[
+				CategoryToolbarVBox.ToSharedRef()
+			]
+
+		+ SSplitter::Slot()
+		.SizeRule(SSplitter::FractionOfParent)
+			[
+				ToolkitWidgetVBox->AsShared()
+			];
+	
+		ToolkitWidgetContainerVBox->AddSlot()
+    	.VAlign(VAlign_Fill)
+	    .FillHeight(1)
+		[
+			MainSplitter->AsShared()
+		];
+}
+
+FCategoryDrivenContentBuilderArgs::FCategoryDrivenContentBuilderArgs(FName InBuilderName):
+                                                                                                      BuilderName( InBuilderName ),
+                                                                                                      bShowCategoryButtonLabels(false),
+                                                                                                      CategoryReclickBehavior(FCategoryDrivenContentBuilderBase::ECategoryReclickBehavior::NoEffect)
+{
+}
+
+void FCategoryDrivenContentBuilderBase::SetCategoryButtonLabelVisibility(EVisibility Visibility)
+{
+
+	CategoryButtonLabelVisibility = Visibility;
+	InitializeCategoryToolbar();
+}
+
+void FCategoryDrivenContentBuilderBase::SetCategoryButtonLabelVisibility(bool bIsCategoryButtonLabelVisible)
+{
+	SetCategoryButtonLabelVisibility(bIsCategoryButtonLabelVisible ? EVisibility::Visible : EVisibility::Collapsed);
+}
+
+
+TSharedRef<SWidget> FCategoryDrivenContentBuilderBase::CreateToolbarWidget() const
+{
+	return ToolRegistry.GenerateWidget(VerticalToolbarElement.ToSharedRef());
+}
\ No newline at end of file
diff --git a/Engine/Source/Developer/WidgetRegistration/Private/ToolElementRegistry.cpp b/Engine/Source/Developer/WidgetRegistration/Private/ToolElementRegistry.cpp
--- a/Engine/Source/Developer/WidgetRegistration/Private/ToolElementRegistry.cpp
+++ b/Engine/Source/Developer/WidgetRegistration/Private/ToolElementRegistry.cpp
@@ -32,6 +32,15 @@ TSharedPtr<SWidget> FToolElementRegistrationArgs::GenerateWidget()
 	return nullptr;
 }
 
+TSharedRef<SWidget> FToolElementRegistrationArgs::GenerateWidgetSharedRef()
+{
+	if (TSharedPtr<SWidget> Widget = GenerateWidget())
+	{
+		return Widget.ToSharedRef();
+	}
+	return SNullWidget::NullWidget;
+}
+
 FToolElementRegistrationKey::FToolElementRegistrationKey(FName InName, EToolElement InToolElementType ) :
 		Name(InName), ToolElementType(InToolElementType)
 {
diff --git a/Engine/Source/Developer/WidgetRegistration/Private/ToolkitBuilder.cpp b/Engine/Source/Developer/WidgetRegistration/Private/ToolkitBuilder.cpp
--- a/Engine/Source/Developer/WidgetRegistration/Private/ToolkitBuilder.cpp
+++ b/Engine/Source/Developer/WidgetRegistration/Private/ToolkitBuilder.cpp
@@ -12,11 +12,10 @@
 #include "Layout/SeparatorBuilder.h"
 #include "Layout/SeparatorTemplates.h"
 #include "Styling/StyleColors.h"
+#include "Layout/CategoryDrivenContentBuilderBase.h"
 
 #define LOCTEXT_NAMESPACE "ToolkitBuilder"
 
-FToolElementRegistry FToolkitBuilder::ToolRegistry = FToolElementRegistry::Get();
-
 bool FEditablePalette::IsInPalette(const FName CommandName) const
 {
 	return this->PaletteCommandNameArray.Contains(CommandName.ToString());
@@ -75,29 +74,35 @@ void FEditablePalette::LoadFromConfig()
 	}
 }
 
+FToolkitBuilderArgs::FToolkitBuilderArgs(const FName InToolbarCustomizationName):
+	FCategoryDrivenContentBuilderArgs(InToolbarCustomizationName)
+	, ToolbarCustomizationName(InToolbarCustomizationName)
+{
+}
+
 FEditablePalette::FEditablePalette(TSharedPtr<FUICommandInfo> InLoadToolPaletteAction,
-	TSharedPtr<FUICommandInfo> InAddToPaletteAction,
-	TSharedPtr<FUICommandInfo> InRemoveFromPaletteAction,
-	FName InEditablePaletteName,
-	FGetEditableToolPaletteConfigManager InGetConfigManager) :
-	FToolPalette(InLoadToolPaletteAction, {}),
-	AddToPaletteAction(InAddToPaletteAction),
-	RemoveFromPaletteAction(InRemoveFromPaletteAction),
-	EditablePaletteName(InEditablePaletteName),
-	GetConfigManager(InGetConfigManager)
+                                   TSharedPtr<FUICommandInfo> InAddToPaletteAction,
+                                   TSharedPtr<FUICommandInfo> InRemoveFromPaletteAction,
+                                   FName InEditablePaletteName,
+                                   FGetEditableToolPaletteConfigManager InGetConfigManager) :
+                                                                                            FToolPalette(InLoadToolPaletteAction, {}),
+                                                                                            AddToPaletteAction(InAddToPaletteAction),
+                                                                                            RemoveFromPaletteAction(InRemoveFromPaletteAction),
+                                                                                            EditablePaletteName(InEditablePaletteName),
+                                                                                            GetConfigManager(InGetConfigManager)
 {
 	LoadFromConfig();
 }
 
-FToolkitBuilder::FToolkitBuilder(const FToolkitBuilderArgs& Args) :
-	FToolElementRegistrationArgs(EToolElement::Toolkit),
-	ToolbarCustomizationName(Args.ToolbarCustomizationName),
-	ToolkitCommandList(Args.ToolkitCommandList),
-	ToolkitSections(Args.ToolkitSections),
-	SelectedCategoryTitleVisibility(Args.SelectedCategoryTitleVisibility),
-	CategoryReclickBehavior(Args.CategoryReclickBehavior)
+
+FToolkitBuilder::FToolkitBuilder(FToolkitBuilderArgs& Args):
+                                                                 FCategoryDrivenContentBuilderBase(Args),
+                                                                 ToolkitCommandList(Args.ToolkitCommandList),
+                                                                 SelectedCategoryTitleVisibility(Args.SelectedCategoryTitleVisibility),
+                                                                 ToolkitSections(Args.ToolkitSections)
 {
-	SetCategoryButtonLabelVisibility(Args.bShowCategoryButtonLabels);
+	const bool bInitLoadToolPaletteMap = true;
+	InitializeCategoryToolbar( bInitLoadToolPaletteMap );
 	ResetWidget();
 }
 
@@ -107,32 +112,6 @@ FToolkitBuilder::~FToolkitBuilder()
 	{
 		ToolRegistry.UnregisterElement(PaletteElement);
 	}
-	if (VerticalToolbarElement.IsValid())
-	{
-		ToolRegistry.UnregisterElement(VerticalToolbarElement.ToSharedRef());
-	}
-}
-
-FToolkitBuilder::FToolkitBuilder(
-	FName InToolbarCustomizationName,
-	TSharedPtr<FUICommandList> InToolkitCommandList,
-	TSharedPtr<FToolkitSections> InToolkitSections) :
-	FToolElementRegistrationArgs(EToolElement::Toolkit),
-	ToolbarCustomizationName(InToolbarCustomizationName),
-	ToolkitCommandList(InToolkitCommandList),
-	ToolkitSections(InToolkitSections)
-{
-	ResetWidget();
-}
-
-TSharedPtr<FToolBarBuilder> FToolkitBuilder::GetLoadPaletteToolbar()
-{
-	return LoadPaletteToolBarBuilder;
-}
-
-TSharedRef<SWidget> FToolkitBuilder::CreateToolbarWidget() const
-{
-	return ToolRegistry.GenerateWidget(VerticalToolbarElement.ToSharedRef());
 }
 
 void FToolkitBuilder::GetCommandsForEditablePalette(TSharedRef<FEditablePalette> EditablePalette, TArray<TSharedPtr<const FUICommandInfo>>& OutCommands)
@@ -185,7 +164,7 @@ void FToolkitBuilder::AddPalette(TSharedPtr<FToolPalette> Palette)
 				FExecuteAction::CreateSP(SharedThis(this),
 				&FToolkitBuilder::TogglePalette,
 				Palette),
-				FCanExecuteAction::CreateLambda([] { return true; }),
+				FCanExecuteAction(),
 				FGetActionCheckState::CreateSP(SharedThis(this),
 					&FToolkitBuilder::IsActiveToolPalette,
 					Palette->LoadToolPaletteAction->GetCommandName())
@@ -255,15 +234,23 @@ bool FToolkitBuilder::HasActivePalette() const
 	return ActivePalette != nullptr;
 }
 
-void FToolkitBuilder::InitializeCategoryToolbar(bool InitLoadToolPaletteMap)
+void FToolkitBuilder::InitializeCategoryToolbar()
+{
+	const bool bInitializeToolbar = false;
+	InitializeCategoryToolbar( bInitializeToolbar );
+}
+
+void FToolkitBuilder::InitializeCategoryToolbar(bool bInitLoadToolPaletteMap)
 {
+
+	const bool bForceSmallIcons = true; 
 	Style = FToolkitStyle::Get().GetWidgetStyle<FToolkitWidgetStyle>("FToolkitWidgetStyle");
 	LoadToolPaletteCommandList = MakeShareable(new FUICommandList);
-	LoadPaletteToolBarBuilder = MakeShared<FVerticalToolBarBuilder>(LoadToolPaletteCommandList, FMultiBoxCustomization::None, TSharedPtr<FExtender>(), true);
+	LoadPaletteToolBarBuilder = MakeShared<FVerticalToolBarBuilder>(LoadToolPaletteCommandList, FMultiBoxCustomization::None, TSharedPtr<FExtender>(), bForceSmallIcons);
 	LoadPaletteToolBarBuilder->SetLabelVisibility( CategoryButtonLabelVisibility );
 	EditablePalettesArray.Reset();
 
-	if (InitLoadToolPaletteMap)
+	if (bInitLoadToolPaletteMap)
 	{
 		LoadCommandNameToPaletteToolbarBuilderMap.Reset();
 	}
@@ -283,43 +270,6 @@ void FToolkitBuilder::InitializeCategoryToolbar(bool InitLoadToolPaletteMap)
 
 }
 
-void FToolkitBuilder::InitCategoryToolbarContainerWidget()
-{
-	if (!CategoryToolbarVBox.IsValid())
-	{
-		CategoryToolbarVBox = SNew(SVerticalBox)
-			.Visibility(CategoryToolbarVisibility);
-	}
-	else
-	{
-		CategoryToolbarVBox->ClearChildren();
-	}
-	CategoryToolbarVBox->AddSlot()
-	.Padding(0.f)
-	[
-		CreateToolbarWidget()
-	];
-}
-
-void FToolkitBuilder::RefreshCategoryToolbarWidget()
-{
-	FToolElementRegistrationKey Key = FToolElementRegistrationKey(ToolbarCustomizationName, EToolElement::Toolbar);
-	VerticalToolbarElement = ToolRegistry.GetToolElementSP(Key);
-	const TSharedRef<FToolbarRegistrationArgs> VerticalToolbarRegistrationArgs = MakeShareable<FToolbarRegistrationArgs>(
-		new FToolbarRegistrationArgs(LoadPaletteToolBarBuilder.ToSharedRef()));
-	
-	if (!VerticalToolbarElement.IsValid())
-	{
-		VerticalToolbarElement = MakeShareable(new FToolElement
-			(ToolbarCustomizationName,
-			VerticalToolbarRegistrationArgs));
-		ToolRegistry.RegisterElement(VerticalToolbarElement.ToSharedRef());
-	}
-
-	VerticalToolbarElement->SetRegistrationArgs(VerticalToolbarRegistrationArgs);
-	InitCategoryToolbarContainerWidget();
-}
-
 void FToolkitBuilder::TogglePalette(TSharedPtr<FToolPalette> Palette)
 {
 	const FName CommandName = Palette->LoadToolPaletteAction->GetCommandName();
@@ -459,7 +409,8 @@ TSharedRef<SWidget> FToolkitBuilder::GetContextMenuContent(const FName CommandNa
 
 void FToolkitBuilder::ResetWidget()
 {
-	InitializeCategoryToolbar(true);
+	const bool bInitLoadToolPaletteMap = true;
+	InitializeCategoryToolbar( bInitLoadToolPaletteMap );
 	ToolPaletteWidget = SNew(SVerticalBox);
 }
 
@@ -486,15 +437,6 @@ void FToolkitBuilder::SetActivePaletteOnLoad(const FUICommandInfo* Command)
 	}
 }
 
-TSharedPtr<SWidget> FToolkitBuilder::GenerateWidget()
-{
-	if (!ToolkitWidgetContainerVBox)
-	{
-		DefineWidget();
-	}
-	return ToolkitWidgetContainerVBox.ToSharedRef();
-}
-
 void FToolkitBuilder::SetActiveToolDisplayName(FText InActiveToolDisplayName)
 {
 	ActiveToolDisplayName = InActiveToolDisplayName;
@@ -510,40 +452,15 @@ EVisibility FToolkitBuilder::GetActiveToolTitleVisibility() const
 	return ActiveToolDisplayName.IsEmpty() ? EVisibility::Collapsed : EVisibility::Visible;
 }
 
-void FToolkitBuilder::DefineWidget()
+void FToolkitBuilder::ProvideSelectedCategoryContent( FName ActiveCategoryName )
 {
-	RefreshCategoryToolbarWidget();
-
 	TSharedPtr<SHorizontalBox> ToolNameHeaderBox;
-
-	ToolkitWidgetContainerVBox = SNew(SVerticalBox)
-	+ SVerticalBox::Slot().AutoHeight() [ *FSeparatorTemplates::SmallHorizontalPanelNoBorder()  ]
-	+ SVerticalBox::Slot().AutoHeight() [ *FSeparatorTemplates::SmallHorizontalBackgroundNoBorder() ];
-
-	ToolkitWidgetVBox = SNew(SVerticalBox);
-	TSharedPtr<SWidget> MainSplitter = 
-		SNew(SSplitter)
-		.PhysicalSplitterHandleSize(2.0f)
-		+ SSplitter::Slot()
-		.Resizable(false)
-		.SizeRule(SSplitter::SizeToContent)
-			[
-				CategoryToolbarVBox.ToSharedRef()
-			]
-
-		+ SSplitter::Slot()
-		.SizeRule(SSplitter::FractionOfParent)
-			[
-				ToolkitWidgetVBox->AsShared()
-			];
 	
-	ToolkitWidgetContainerVBox->AddSlot()
-	.VAlign(VAlign_Fill)
-	.FillHeight(1)
-	[
-		MainSplitter->AsShared()
-	];
-
+	if ( !ToolkitWidgetVBox.IsValid() )
+	{
+		return;
+	}
+	
 	if (ToolkitSections->ModeWarningArea)
 	{
 		ToolkitWidgetVBox->AddSlot()
@@ -649,20 +566,14 @@ void FToolkitBuilder::DefineWidget()
 	}
 }
 
-void FToolkitBuilder::SetCategoryButtonLabelVisibility(EVisibility Visibility)
-{
-	CategoryButtonLabelVisibility = Visibility;
-	InitializeCategoryToolbar();
-}
-
-void FToolkitBuilder::SetCategoryButtonLabelVisibility(bool bIsCategoryButtonLabelVisible)
+void FToolkitBuilder::SetActivePaletteCommandsVisibility(EVisibility Visibility)
 {
-	SetCategoryButtonLabelVisibility(bIsCategoryButtonLabelVisible ? EVisibility::Visible : EVisibility::Collapsed);
+	ActivePaletteButtonVisibility = Visibility;
 }
 
-void FToolkitBuilder::SetActivePaletteCommandsVisibility(EVisibility Visibility)
+EVisibility FToolkitBuilder::GetActivePaletteCommandsVisibility() const
 {
-	ActivePaletteButtonVisibility = Visibility;
+	return ActivePaletteButtonVisibility;
 }
 
 #undef LOCTEXT_NAMESPACE
\ No newline at end of file
diff --git a/Engine/Source/Developer/WidgetRegistration/Public/Layout/CategoryDrivenContentBuilder.h b/Engine/Source/Developer/WidgetRegistration/Public/Layout/CategoryDrivenContentBuilder.h
new file mode 100644
--- /dev/null
+++ b/Engine/Source/Developer/WidgetRegistration/Public/Layout/CategoryDrivenContentBuilder.h
@@ -0,0 +1,77 @@
+// Copyright Epic Games, Inc. All Rights Reserved.
+
+#pragma once
+#include "ToolkitBuilder.h"
+#include "ToolkitStyle.h"
+#include "Layout\CategoryDrivenContentBuilderBase.h"
+
+/**
+ * A builder which creates a widget that has a vertical toolbar category picker on the left
+ * hand side which populates the content on the right side.
+ */
+class  FCategoryDrivenContentBuilder : public FCategoryDrivenContentBuilderBase
+{
+public:
+	DECLARE_DELEGATE_RetVal_OneParam( TSharedRef<SWidget>, FProvideSelectedCategoryContent, const FName&)
+
+	/**
+	 * A delegate which provides the const FName& Command name as a parameter to indicates which category was clicked, and returns
+	 * a TSharedRef<SWidget> meant to populate the full content for that category
+	 */ 
+	FProvideSelectedCategoryContent  ProvideSelectedCategoryContentDelegate;
+	
+public:
+	/**
+	 * initializes this FCategoryDrivenContentBuilder with the given FCategoryDrivenContentBuilderArgs
+	 * @param Args the parameter object which provides the information to initialize this FCategoryDrivenContentBuilder 
+	 */
+	WIDGETREGISTRATION_API explicit FCategoryDrivenContentBuilder(FCategoryDrivenContentBuilderArgs& Args);
+	
+	/**
+	 * destroys the FCategoryDrivenContentBuilder and frees any resources
+	 */
+	WIDGETREGISTRATION_API virtual ~FCategoryDrivenContentBuilder() override;
+
+private:
+	/**
+	 * Given the active category name, update the  content
+	 * 
+	 * @param ActiveCategoryName 
+	 */
+	WIDGETREGISTRATION_API virtual void ProvideSelectedCategoryContent( FName ActiveCategoryName = NAME_None ) override;
+
+	FName ActiveCommandName;
+
+public:
+	/*
+	 * Initializes the data necessary to build the category toolbar
+	 */
+	WIDGETREGISTRATION_API virtual void InitializeCategoryToolbar() override;
+
+	/*
+	 * Refreshes the existing widget.
+	 */
+	WIDGETREGISTRATION_API virtual void UpdateWidget() override;
+
+	/**
+	 * Returns true if the FUICommandInfo with the name CommandName is the active tool palette,
+	 * else it returns false
+	 *
+	 * @param CommandName the name of the FUICommandInfo we are checking to see if it is the active tool palette
+	 */
+	WIDGETREGISTRATION_API virtual ECheckBoxState IsActiveToolPalette(FName CommandName) const override
+	{
+		return ActiveCommandName == CommandName ? ECheckBoxState::Checked : ECheckBoxState::Unchecked;
+	}
+
+	/*
+	 * Sets the commands that will load the categories
+	 */
+	WIDGETREGISTRATION_API void SetCommands(TArray<TSharedPtr<FUICommandInfo>> InContentLoaderCommands);
+	
+private:
+	TArray<TSharedPtr<FUICommandInfo>> ContentLoaderCommands;
+};
+
+
+ 
\ No newline at end of file
diff --git a/Engine/Source/Developer/WidgetRegistration/Public/Layout/CategoryDrivenContentBuilderBase.h b/Engine/Source/Developer/WidgetRegistration/Public/Layout/CategoryDrivenContentBuilderBase.h
new file mode 100644
--- /dev/null
+++ b/Engine/Source/Developer/WidgetRegistration/Public/Layout/CategoryDrivenContentBuilderBase.h
@@ -0,0 +1,178 @@
+// Copyright Epic Games, Inc. All Rights Reserved.
+
+#pragma once
+
+
+#include "FToolkitWidgetStyle.h"
+#include "ToolElementRegistry.h"
+#include "Framework/MultiBox/MultiBoxBuilder.h"
+
+
+class SWidget;
+class FToolElement;
+
+struct FCategoryDrivenContentBuilderArgs;
+
+/**
+ * The FToolElementRegistrationArgs which is specified for Toolkits
+ */
+class WIDGETREGISTRATION_API FCategoryDrivenContentBuilderBase : public FToolElementRegistrationArgs
+{
+public:
+
+	// Used to specify what happens when you click the category button of a category that is already active.
+	enum class ECategoryReclickBehavior : uint8
+	{
+		// Do nothing if the same category button is clicked.
+		NoEffect,
+
+		// Toggle the active category off, so no category is active.
+		ToggleOff,
+
+		// Do the same thing that would be done if we were switching from a different category. Note that
+		// this will trigger OnActivePaletteChanged.
+		TreatAsChanged,
+	};
+
+	/**
+	 * A constructor for the FCategoryDrivenContentBuilderBase which defines the builder name
+	 * @param InBuilderName the name for the builder
+	 */
+	explicit  FCategoryDrivenContentBuilderBase( FName InBuilderName );
+
+	/**
+	 * A constructor for the FCategoryDrivenContentBuilderBase which takes a
+ 	* FCategoryDrivenContentBuilderArgs to initialize this
+ 	* @param Args the FCategoryDrivenContentBuilderArgs used to initialize this FCategoryDrivenContentBuilderBase
+ 	*/
+	explicit FCategoryDrivenContentBuilderBase( FCategoryDrivenContentBuilderArgs& Args );
+
+	virtual  ~FCategoryDrivenContentBuilderBase() override;
+
+	/*
+	 * Initializes the data necessary to build the category toolbar
+	 */
+	virtual void InitializeCategoryToolbar() = 0;
+
+	/**
+	 * Initializes the category toolbar container VBox, and the children inside it.
+	 * On any repeat calls, the SVerticalBox created on the first pass will be
+	 * emptied and the children repopulated.
+	 */
+	void InitCategoryToolbarContainerWidget();
+
+	/**
+	 * Sets category button label visibility to Visibility. It also reinitializes the
+	 * category toolbar data, as the toolbar's LabelVisibility member is now stale.
+	 *
+	 * @param Visibility If Visibility == EVisibility::Collapsed, the category button labels
+	 * will be shown, else they will not be shown.
+	 */
+	void SetCategoryButtonLabelVisibility(EVisibility Visibility);
+
+	/**
+	 * Sets category button label visibility to Visiblity. It also reinitializes the
+	 * category toolbar data, as the toolbar's LabelVisibility member is now stale.
+	 *
+	 * @param bIsCategoryButtonLabelVisible If bIsCategoryButtonLabelVisible == true,
+	 * the category button labels will be shown, else they will not be shown.
+	 */
+	void SetCategoryButtonLabelVisibility(bool bIsCategoryButtonLabelVisible);
+
+	/**
+	 * RefreshCategoryToolbarWidget refreshes the UI display of the category toolbar 
+	 */
+	void RefreshCategoryToolbarWidget();
+
+	/** Implements the generation of the TSharedPtr<SWidget> */
+	virtual TSharedPtr<SWidget> GenerateWidget() override;
+
+	/** Creates the Toolbar for the widget with the FUICommandInfos that load the Palettes */
+	TSharedRef<SWidget> CreateToolbarWidget() const;
+
+	/** returns a TSharedPointer to the FToolbarBuilder with the FUICommandInfos that load the Palettes */
+	TSharedPtr<FToolBarBuilder> GetLoadPaletteToolbar();
+	
+	/** a TSharedPointer to the FToolElement for the vertical toolbar */
+	TSharedPtr<FToolElement> VerticalToolbarElement;
+
+	/**   OnActivePaletteChanged is broadcast when the active palette changes to a different palette  */
+	FSimpleMulticastDelegate OnActivePaletteChanged;
+	
+private:
+
+	/* The SWidget that is the whole Toolkit */
+	TSharedPtr<SWidget> ToolkitWidget;
+
+	/**
+	 * Returns true if the FUICommandInfo with the name CommandName is the active tool palette,
+	 * else it returns false
+	 *
+	 * @param CommandName the name of the FUICommandInfo we are checking to see if it is the active tool palette
+	 */
+	virtual ECheckBoxState IsActiveToolPalette(FName CommandName) const = 0;
+
+	/**
+	 * Creates the SWidget
+	 */
+	void CreateWidget();
+
+	/** The SVerticalBox which contains the category toolbar */
+	TSharedPtr<SVerticalBox> CategoryToolbarVBox;
+
+	/** The SVerticalBox which holds the entire Toolkit */
+	TSharedPtr<SVerticalBox> ToolkitWidgetContainerVBox;
+
+protected:
+	virtual void ProvideSelectedCategoryContent( FName ActiveCategoryName = NAME_None) = 0;
+	
+	/** The SVerticalBox which holds all but the vertical toolbar in a Toolkit */
+	TSharedPtr<SVerticalBox> ToolkitWidgetVBox;
+	
+	/** the tool element registry this class will use to register UI tool elements */
+	static FToolElementRegistry ToolRegistry;
+	
+	/** Specifies what happens if you click the category button of an already-active catogory */
+	ECategoryReclickBehavior CategoryReclickBehavior = ECategoryReclickBehavior::NoEffect;
+	
+	/** The current FToolkitWidgetStyle */
+	FToolkitWidgetStyle Style;
+	
+	/** If CategoryButtonLabelVisibility == EVisibility::Visible, the category button
+	 * labels are visible, else they are not displayed. By default the selected category button labels are Visible    */	
+	EVisibility CategoryButtonLabelVisibility;
+	
+	/*
+	 * Is the Category toolbar visible? This should be false if the Category toolbar has less than 2 categories,
+	 * because if there is only one showing it is superfluous to actually show it since it will always be selected.
+	 */
+	EVisibility CategoryToolbarVisibility;
+	
+	/** A TSharedPointer to the FUICommandList for the FUICommandInfos which load a tool palette */
+	TSharedPtr<FUICommandList> LoadToolPaletteCommandList;
+	
+	/** The toolbar builder for the toolbar which has the FUICommandInfos which load the various palettes */
+	TSharedPtr<FToolBarBuilder> LoadPaletteToolBarBuilder;
+
+	/** Name of this builder */
+	FName BuilderName;
+};
+
+/**
+ * A struct to provide arguments for a FCategoryDrivenContentBuilderBase
+ */
+struct FCategoryDrivenContentBuilderArgs
+{
+public:
+	WIDGETREGISTRATION_API FCategoryDrivenContentBuilderArgs( FName InBuilderName );
+
+	/** Name of this builder */
+	FName BuilderName;
+	
+	/** If bShowCategoryButtonLabels == true, the category button
+   * labels should be visible, else they are not displayed  */
+	bool bShowCategoryButtonLabels;
+
+	/** Specifies what happens if you click the category button of an already-active category */
+	FCategoryDrivenContentBuilderBase::ECategoryReclickBehavior CategoryReclickBehavior;
+};
\ No newline at end of file
diff --git a/Engine/Source/Developer/WidgetRegistration/Public/ToolElementRegistry.h b/Engine/Source/Developer/WidgetRegistration/Public/ToolElementRegistry.h
--- a/Engine/Source/Developer/WidgetRegistration/Public/ToolElementRegistry.h
+++ b/Engine/Source/Developer/WidgetRegistration/Public/ToolElementRegistry.h
@@ -36,6 +36,9 @@ public:
 	/** Implements the generation of the TSharedPtr<SWidget> */
 	virtual TSharedPtr<SWidget> GenerateWidget();
 
+	/** Implements the generation of the TSharedRef<SWidget> */
+	virtual TSharedRef<SWidget> GenerateWidgetSharedRef();
+
 	/** Updates/reloads this widget. This should be called after a consumer has changed any Data in this */
 	virtual void UpdateWidget();
 
diff --git a/Engine/Source/Developer/WidgetRegistration/Public/ToolkitBuilder.h b/Engine/Source/Developer/WidgetRegistration/Public/ToolkitBuilder.h
--- a/Engine/Source/Developer/WidgetRegistration/Public/ToolkitBuilder.h
+++ b/Engine/Source/Developer/WidgetRegistration/Public/ToolkitBuilder.h
@@ -2,16 +2,14 @@
 
 #pragma once
 
-#include "FToolkitWidgetStyle.h"
 #include "IDetailsView.h"
 #include "Framework/MultiBox/SToolBarButtonBlock.h"
 #include "Framework/Commands/UICommandInfo.h"
 #include "ToolElementRegistry.h"
 #include "ToolkitBuilderConfig.h"
-#include "ToolMenus.h"
 #include "Widgets/SBoxPanel.h"
 #include "Widgets/Layout/SSplitter.h"
-
+#include "Layout/CategoryDrivenContentBuilderBase.h"
 
 class SWidget;
 class FToolElement;
@@ -100,89 +98,53 @@ protected:
 
 struct FToolkitBuilderArgs;
 
-/**
- * The FToolElementRegistrationArgs which is specified for Toolkits
+/*
+ * A simple struct to carry initialization data for an FToolkitBuilder
  */
-class WIDGETREGISTRATION_API FToolkitBuilder : public FToolElementRegistrationArgs
+struct FToolkitBuilderArgs : public FCategoryDrivenContentBuilderArgs
 {
-public:
-	FToolkitBuilder(
-		FName ToolbarCustomizationName,
-		TSharedPtr<FUICommandList> ToolkitCommandList,
-		TSharedPtr<FToolkitSections> ToolkitSections);
-
-	FToolkitBuilder(const FToolkitBuilderArgs& Args);
-	
-	virtual  ~FToolkitBuilder() override;
-	
-	// Used to specify what happens when you click the category button of a category that is already active.
-	enum class ECategoryReclickBehavior : uint8
-	{
-		// Do nothing if the same category button is clicked.
-		NoEffect,
+	WIDGETREGISTRATION_API explicit FToolkitBuilderArgs(const FName InToolbarCustomizationName);
 
-		// Toggle the active category off, so no category is active.
-		ToggleOff,
+	/** Name of the toolbar this mode uses and can be used by external systems to customize that mode toolbar */
+	FName ToolbarCustomizationName;
 
-		// Do the same thing that would be done if we were switching from a different category. Note that
-		// this will trigger OnActivePaletteChanged.
-		TreatAsChanged,
-	};
+	/** A TSharedPointer to the FUICommandList for the current mode */
+	TSharedPtr<FUICommandList> ToolkitCommandList;
 
-	/*
-	 * Initializes the data to necessary to build the category toolbar
-	 *
-	 * @param InitLoadToolPaletteMap If true, the map which holds the load
-	 * action names as keys and the ToolPalettes as value will be cleared and available
-	 * to add new data when the method completes. Otherwise, the map will keep
-	 * the load action name to ToolPalette mappings that it was initially set up with. 
-	 */
-	void InitializeCategoryToolbar(bool InitLoadToolPaletteMap = false);
+	/** The FToolkitSections which holds the sections defined for this Toolkit */
+	TSharedPtr<FToolkitSections> ToolkitSections;
 
-	/**
-	 * Initializes the category toolbar container VBox, and the children inside it.
-	 * On any repeat calls, the SVerticalBox created on the first pass will be
-	 * emptied and the children repopulated.
-	 */
-	void InitCategoryToolbarContainerWidget();
+	/** If SelectedCategoryTitleVisibility == EVisibility::Visible, the selected
+	 * category title is visible, else it is not displayed  */
+	EVisibility SelectedCategoryTitleVisibility;
+};
 
+class WIDGETREGISTRATION_API FToolkitBuilder : public FCategoryDrivenContentBuilderBase
+{
+public:
 	/**
-	 * Sets category button label visibility to Visibility. It also reinitializes the
-	 * category toolbar data, as the toolbar's LabelVisibility member is now stale.
-	 *
-	 * @param Visibility If Visibility == EVisibility::Collapsed, the category button labels
-	 * will be shown, else they will not be shown.
+	 * Constructor
+	 * 
+	 * @param ToolbarCustomizationName the name of the  customization fr the category toolbar 
+	 * @param InToolkitCommandList  the toolkit FUICommandList
+	 * @param InToolkitSections The FToolkitSections for this toolkit builder
 	 */
-	void SetCategoryButtonLabelVisibility(EVisibility Visibility);
+	explicit FToolkitBuilder(
+	FName ToolbarCustomizationName,
+	TSharedPtr<FUICommandList> InToolkitCommandList,
+	TSharedPtr<FToolkitSections> InToolkitSections);
 
-	/**
-	 * Sets category button label visibility to Visiblity. It also reinitializes the
-	 * category toolbar data, as the toolbar's LabelVisibility member is now stale.
-	 *
-	 * @param bIsCategoryButtonLabelVisible If bIsCategoryButtonLabelVisible == true,
-	 * the category button labels will be shown, else they will not be shown.
-	 */
-	void SetCategoryButtonLabelVisibility(bool bIsCategoryButtonLabelVisible);
+	/** default constructor */
+	explicit FToolkitBuilder(FToolkitBuilderArgs& Args);
 
+	virtual ~FToolkitBuilder() override;
 
-	/**
-	 * Configure whether or not the tool buttons in the active palette should be visible.
-	 * This can be dynamically toggled without rebuilding/regenerating the widget.
-	 * 
-	 * @param Visibility If Visibility == EVisibility::Collapsed, the category button labels
-	 * will be shown, else they will not be shown.
-	 */
 	void SetActivePaletteCommandsVisibility(EVisibility Visibility);
 
 	/**
 	 * @return the Visibility state of the active palette command buttons
 	 */
-	EVisibility GetActivePaletteCommandsVisibility() const { return ActivePaletteButtonVisibility; }
-
-	/**
-	 * RefreshCategoryToolbarWidget refreshes the UI display of the category toolbar 
-	 */
-	void RefreshCategoryToolbarWidget();
+	EVisibility GetActivePaletteCommandsVisibility() const;
 
 	/**
 	 * Adds the FToolPalette Palette to this FToolkitBuilder
@@ -199,30 +161,12 @@ public:
 	void AddPalette(TSharedPtr<FEditablePalette> Palette);
 
 	/** Cleans up any previously set data in this FToolkitBuilder and reset the members to their initial values  */
-	virtual void ResetWidget() override;
+	virtual void ResetWidget() override final;
 
 	/** Updates the Toolkit. This should be called after any changes to the data of this FToolkitBuilder, and the UI
 	 * will be regenerated to reflect it.  */
 	virtual void UpdateWidget() override;
 
-	/** Implements the generation of the TSharedPtr<SWidget> */
-	virtual TSharedPtr<SWidget> GenerateWidget() override;
-
-	/* Returns true if the Toolkit builder has some tools that are currently active/selected */
-	bool HasSelectedToolSet() const;
-
-	/** Creates the Toolbar for the widget with the FUICommandInfos that load the Palettes */
-	TSharedRef<SWidget> CreateToolbarWidget() const;
-
-	/** returns a TSharedPointer to the FToolbarBuilder with the FUICommandInfos that load the Palettes */
-	TSharedPtr<FToolBarBuilder> GetLoadPaletteToolbar();
-	
-	/** returns a TSharedPointer to the FToolbarBuilder with the FUICommandInfos that load the Palettes */
-	TSharedPtr<FToolElement> VerticalToolbarElement;
-
-	/** returns a TSharedPointer to the FToolbarBuilder with the FUICommandInfos that load the Palettes */
-	TSharedRef<SWidget> GetToolPaletteWidget() const;
-
 	/** returns true is there is an active palette selected, else it returns false */
 	bool HasActivePalette() const;
 
@@ -240,11 +184,21 @@ public:
 	 * @param CommandName the name of the FUICommandInfo we are checking to see if it is the active tool palette
 	 */
 	ECheckBoxState IsActiveToolPalette(FName CommandName) const;
+	
+	/*
+	 * Initializes the data necessary to build the category toolbar
+	 */
+	void InitializeCategoryToolbar() override;
 
-	/** 
-	 * OnActivePaletteChanged is broadcast when the active palette changes to a different palette 
+	/**
+	 * Initializes the data necessary to build the category toolbar
+	 *
+	 * @param bInitLoadToolPaletteMap If true, the map which holds the load
+	 * action names as keys and the ToolPalettes as value will be cleared and available
+	 * to add new data when the method completes. Otherwise, the map will keep
+	 * the load action name to ToolPalette mappings that it was initially set up with. 
 	 */
-	FSimpleMulticastDelegate OnActivePaletteChanged;
+	void InitializeCategoryToolbar(bool bInitLoadToolPaletteMap);
 
 	/**
 	 * Sets the display name for the active tool
@@ -270,15 +224,7 @@ public:
 	 * @returns the name of the Active Palette if one is available, else it returns NAME_NONE
 	 */
 	FName GetActivePaletteName() const;
-	
-private:
-
-	/** the tool element registry this class will use to register UI tool elements */
-	static FToolElementRegistry ToolRegistry;
 
-	/* The SWidget that is the whole Toolkit */
-	TSharedPtr<SWidget> ToolkitWidget;
-	
 	/** Name of the toolbar this mode uses and can be used by external systems to customize that mode toolbar */
 	FName ToolbarCustomizationName;
 
@@ -290,10 +236,7 @@ private:
 
 	/** Map of command name to the actual command for all commands belonging to this palette */
 	TMap<FString, TSharedPtr<const FUICommandInfo>> PaletteCommandInfos;
-
-	/** A TSharedPointer to the FUICommandList for the FUICommandInfos which load a tool palette */
-	TSharedPtr<FUICommandList> LoadToolPaletteCommandList;
-
+	
 	/** A TSharedPointer to the FUICommandList for the current mode */
 	TSharedPtr<FUICommandList> ToolkitCommandList;
 
@@ -305,10 +248,7 @@ private:
 
 	/** The SVerticalBox which holds the tool palette (the Buttons which load each tool) */
 	TSharedPtr<SVerticalBox> ToolPaletteWidget;
-
-	/** The toolbar builder for the toolbar which has the FUICommandInfos which load the various palettes */
-	TSharedPtr<FToolBarBuilder> LoadPaletteToolBarBuilder;
-
+	
 	/** the FName of each Load palette command to the FToolbarBuilder for the palette which it loads */
 	TMap<FName, TSharedPtr<FToolBarBuilder>> LoadCommandNameToPaletteToolbarBuilderMap;
 
@@ -320,6 +260,11 @@ private:
 	/** Resets the tool palette widget, on which the buttons for the currently chosen toolset are shown */
 	void ResetToolPaletteWidget();
 
+	/** The array of Tool Palettes on display. This array is used to keep a handle
+	* on all of the created FToolElements so that we can unregister them upon destruction */
+	TArray<TSharedRef<FToolElement>> ToolPaletteElementArray;
+
+private:
 	/**
 	 * Creates the UI for the Palette specified by the FToolPalette Palette
 	 *
@@ -334,7 +279,20 @@ private:
 	 * @param Palette the FToolPalette for which this method will build the UI 
 	 */
 	void TogglePalette(TSharedPtr<FToolPalette> Palette);
-
+	
+	/**
+	 * Retrieves the context menu content for the FUICommandInfo with the command name CommandName
+	 *
+	 * @param CommandName the command name of the FUICommandInfo to get the context menu widget for
+	 * @return the TSharedRef<SWidget> which contains the context menu for the FUICommandInfo with the name CommandName 
+	 */
+	TSharedRef<SWidget> GetContextMenuContent(const FName CommandName);
+	
+	/** returns a TSharedPointer to the FToolbarBuilder with the FUICommandInfos that load the Palettes */
+	TSharedRef<SWidget> GetToolPaletteWidget() const;
+	
+	void CreatePaletteWidget(FToolPalette& Palette, FToolElement& Element);
+	
 	/**
 	 * Adds the FUICommandInfo with the command name CommandNameString to the FEditablePalette Palette if
 	 * it is not on Palette, and removes the FUICommandInfo from Palette if it is on the Palette
@@ -343,14 +301,16 @@ private:
 	 * @param CommandNameString the name of the FUICommandInfo command which will be toggled on/off of Palette
 	 */
 	void ToggleCommandInPalette(TSharedRef<FEditablePalette> Palette, FString CommandNameString);
+		
+	/* Returns true if the Toolkit builder has some tools that are currently active/selected */
+	bool HasSelectedToolSet() const;
 
 	/**
-	 * Retrieves the context menu content for the FUICommandInfo with the command name CommandName
-	 *
-	 * @param CommandName the command name of the FUICommandInfo to get the context menu widget for
-	 * @return the TSharedRef<SWidget> which contains the context menu for the FUICommandInfo with the name CommandName 
+	 * Given the active category name, update the toolkit content
+	 * 
+	 * @param ActiveCategoryName 
 	 */
-	TSharedRef<SWidget> GetContextMenuContent(const FName CommandName);
+	virtual void ProvideSelectedCategoryContent( FName ActiveCategoryName = NAME_None ) override;
 
 	/**
 	 * Gets the EVisibility of the active tool title. This should only be visible if a tool is currently chosen in the palette.
@@ -358,83 +318,18 @@ private:
 	 * @return the EVisibility of the active tool title
 	 */
 	EVisibility GetActiveToolTitleVisibility() const;
-
-	void CreatePaletteWidget(FToolPalette& Palette, FToolElement& Element);
-
-	void DefineWidget();
-
-	/** The SVerticalBox which contains the category toolbar */
-	TSharedPtr<SVerticalBox> CategoryToolbarVBox;
-
-	/** The array of Tool Palettes on display. This array is used to keep a handle
-	 * on all of the created FToolElements so that we can unregister them upon destruction */
-	TArray<TSharedRef<FToolElement>> ToolPaletteElementArray;
-
-	/** The SVerticalBox which holds all but the vertical toolbar in a Toolkit */
-	TSharedPtr<SVerticalBox> ToolkitWidgetVBox;
-
-	/** The SVerticalBox which holds the entire Toolkit */
-	TSharedPtr<SVerticalBox> ToolkitWidgetContainerVBox;
-
-	/** The FToolkitSections which holds the sections defined for this Toolkit */
-	TSharedPtr<FToolkitSections> ToolkitSections;
-
-	/** The display name of the currently active tool */
-	FText ActiveToolDisplayName = FText::GetEmpty();
-
-	/** The current FToolkitWidgetStyle */
-	FToolkitWidgetStyle Style;
-
+	
 	/** If SelectedCategoryTitleVisibility == EVisibility::Visible, the selected
-	 * category title is visible, else it is not displayed. By default the selected category title is Collapsed  */	
+    * category title is visible, else it is not displayed. By default the selected category title is Collapsed  */	
 	EVisibility SelectedCategoryTitleVisibility = EVisibility::Collapsed;
 
-	/** If CategoryButtonLabelVisibility == EVisibility::Visible, the category button
-	 * labels are visible, else they are not displayed. By default the selected category button labels are Visible    */	
-	EVisibility CategoryButtonLabelVisibility;
+	/** The FToolkitSections which holds the sections defined for this Toolkit */
+	TSharedPtr<FToolkitSections> ToolkitSections;
 
 	/** If ActivePaletteButtonVisibility == EVisibility::Visible, the command buttons in the active palette
 	 * are visible, otherwise they are not displayed. By default the state is Visible  */
 	EVisibility ActivePaletteButtonVisibility = EVisibility::Visible;
-
-	/*
-	 * Is the Category toolbar visible? This should be false if the Category toolbar has less than 2 categories,
-	 * because if there is only one showing it is superfluous to actually show it since it will always be selected.
-	 */
-	EVisibility CategoryToolbarVisibility;
-
-	/** Specifies what happens if you click the category button of an already-active catogory */
-	ECategoryReclickBehavior CategoryReclickBehavior = ECategoryReclickBehavior::NoEffect;
-};
-
-//~ This is defined after the class declaration so we can use the nested enum, which there isn't
-//~ a way to forward-declare.
-/*
- * A simple struct to carry initialization data for an FToolkitBuilder
- */
-struct FToolkitBuilderArgs
-{
-	FToolkitBuilderArgs(const FName& InToolbarCustomizationName) : ToolbarCustomizationName(InToolbarCustomizationName)
-	{
-	}
-
-	/** Name of the toolbar this mode uses and can be used by external systems to customize that mode toolbar */
-	const FName& ToolbarCustomizationName;
-
-	/** A TSharedPointer to the FUICommandList for the current mode */
-	TSharedPtr<FUICommandList> ToolkitCommandList;
-
-	/** The FToolkitSections which holds the sections defined for this Toolkit */
-	TSharedPtr<FToolkitSections> ToolkitSections;
-
-	/** If bShowCategoryButtonLabels == true, the category button
-	* labels should be, else they are not displayed  */
-	bool bShowCategoryButtonLabels;
-
-	/** If SelectedCategoryTitleVisibility == EVisibility::Visible, the selected
-	 * category title is visible, else it is not displayed  */
-	EVisibility SelectedCategoryTitleVisibility;
-
-	/** Specifies what happens if you click the category button of an already-active catogory */
-	FToolkitBuilder::ECategoryReclickBehavior CategoryReclickBehavior = FToolkitBuilder::ECategoryReclickBehavior::NoEffect;
-};
+	
+	/** The display name of the currently active tool */
+	FText ActiveToolDisplayName = FText::GetEmpty();
+};
\ No newline at end of file
diff --git a/Engine/Source/Editor/PlacementMode/PlacementMode.Build.cs b/Engine/Source/Editor/PlacementMode/PlacementMode.Build.cs
--- a/Engine/Source/Editor/PlacementMode/PlacementMode.Build.cs
+++ b/Engine/Source/Editor/PlacementMode/PlacementMode.Build.cs
@@ -24,6 +24,7 @@ public class PlacementMode : ModuleRules
                 "LevelEditor",
                 "AssetTools",
                 "ToolMenus",
+                "WidgetRegistration"
             } 
         );
     }
diff --git a/Engine/Source/Editor/PlacementMode/Private/PlacementModeCommands.cpp b/Engine/Source/Editor/PlacementMode/Private/PlacementModeCommands.cpp
new file mode 100644
--- /dev/null
+++ b/Engine/Source/Editor/PlacementMode/Private/PlacementModeCommands.cpp
@@ -0,0 +1,50 @@
+
+// Copyright Epic Games, Inc. All Rights Reserved.
+
+#include "PlacementModeCommands.h"
+#include "IPlacementModeModule.h"
+#include "Framework/Commands/UICommandInfo.h"
+
+#define LOCTEXT_NAMESPACE "PlacementModeCommands"
+
+FPlacementModeCommands::FPlacementModeCommands()
+	: TCommands<FPlacementModeCommands>(
+		TEXT("PlacementModeCommands"), 
+		LOCTEXT( "PlacementModeCommands", "Placement Mode"), 
+		NAME_None, 
+		"FPlacementModeCommandsStyle")
+{}
+
+void FPlacementModeCommands::RegisterCommands()
+{
+	TArray<FPlacementCategoryInfo> Categories;
+	IPlacementModeModule::Get().GetSortedCategories(Categories);
+	
+	for (FPlacementCategoryInfo& Category : Categories)		
+	{
+		TSharedPtr<FUICommandInfo> Command;
+
+		FUICommandInfo::MakeCommandInfo(
+			this->AsShared(),
+			Command,
+			Category.UniqueHandle,
+			Category.DisplayName,
+			Category.DisplayName,
+			Category.DisplayIcon,
+			EUserInterfaceActionType::ToggleButton,
+			FInputChord());
+			ToolkitCategoryCommands.Add( Command );
+	}
+}
+
+const FPlacementModeCommands& FPlacementModeCommands::Get()
+{
+	return TCommands<FPlacementModeCommands>::Get();
+}
+
+const TArray<TSharedPtr<FUICommandInfo>>& FPlacementModeCommands::GetPlacementToolkitCategoryCommands() const
+{
+	return ToolkitCategoryCommands;
+}
+
+#undef LOCTEXT_NAMESPACE
diff --git a/Engine/Source/Editor/PlacementMode/Private/PlacementModeCommands.h b/Engine/Source/Editor/PlacementMode/Private/PlacementModeCommands.h
new file mode 100644
--- /dev/null
+++ b/Engine/Source/Editor/PlacementMode/Private/PlacementModeCommands.h
@@ -0,0 +1,40 @@
+// Copyright Epic Games, Inc. All Rights Reserved.
+
+#pragma once
+
+#include "Framework/Commands/Commands.h"
+
+struct FPlacementCategoryInfo;
+
+/**
+* Creates  the commands for placement mode, creating one command to load each  placement mode category 
+ */
+class  FPlacementModeCommands : public TCommands<FPlacementModeCommands>
+{
+public:
+	/** default constructor  */
+	FPlacementModeCommands();
+
+	/**
+	 * Registers the commands for placement mode, creating one command to load each
+	 * placement mode category 
+	 */
+	virtual void RegisterCommands() override;
+
+	/**
+	 * returns a static instance of the placement mode commands
+	 */
+	static const FPlacementModeCommands& Get();
+
+	/**
+	 * Returns an array containing the commands that load the categories of things to be placed in the Placement Panel
+	 */
+	const TArray<TSharedPtr<FUICommandInfo>>& GetPlacementToolkitCategoryCommands() const;
+
+private:
+	/**
+	 * An array containing the commands that load the categories of things to be placed in the Placement Panel
+	 */
+	TArray<TSharedPtr<FUICommandInfo>> ToolkitCategoryCommands;
+};
+
diff --git a/Engine/Source/Editor/PlacementMode/Private/SPlacementModeTools.cpp b/Engine/Source/Editor/PlacementMode/Private/SPlacementModeTools.cpp
--- a/Engine/Source/Editor/PlacementMode/Private/SPlacementModeTools.cpp
+++ b/Engine/Source/Editor/PlacementMode/Private/SPlacementModeTools.cpp
@@ -24,9 +24,17 @@
 #include "AssetSelection.h"
 #include "ActorFactories/ActorFactory.h"
 #include "ScopedTransaction.h"
+#include "PlacementModeCommands.h"
+#include "Layout/CategoryDrivenContentBuilder.h"
+#include "Widgets/Layout/SScrollBox.h"
 
 #define LOCTEXT_NAMESPACE "PlacementMode"
 
+static TAutoConsoleVariable<bool> CVarEnableCategoryContentChooserView(
+	TEXT("PlaceActors.EnableCategoryChooserView"),
+	false,
+	TEXT("If enabled, the new category chooser view will be used for the Place Actors panel."));
+
 namespace PlacementModeTools
 {
 	bool bItemInternalsInTooltip = false;
@@ -592,6 +600,12 @@ void SPlacementModeTools::Construct( const FArguments& InArgs, TSharedRef<SDockT
 	bRefreshRecentlyPlaced = false;
 	bUpdateShownItems = true;
 
+	FCategoryDrivenContentBuilderArgs Args("PlacementModes");
+	CategoryContentBuilder = MakeShared<FCategoryDrivenContentBuilder>( Args );
+	CategoryContentBuilder->ProvideSelectedCategoryContentDelegate.BindSP(this, &SPlacementModeTools::ProvideCategoryContent);
+	UniformWrapPanel = SNew( SUniformWrapPanel );
+	SAssignNew(CustomContent, SBox);
+
 	ActiveTabName = FBuiltInPlacementCategories::Basic();
 
 	ParentTab->SetOnTabDrawerOpened(FSimpleDelegate::CreateSP(this, &SPlacementModeTools::OnTabDrawerOpened));
@@ -626,7 +640,17 @@ void SPlacementModeTools::Construct( const FArguments& InArgs, TSharedRef<SDockT
 				.OnTextCommitted(this, &SPlacementModeTools::OnSearchCommitted)
 			]
 		]
-
+		+ SVerticalBox::Slot()
+			.FillHeight(1)
+			[
+				SNew(SBorder)
+				.BorderImage(FAppStyle::GetBrush("ToolPanel.GroupBorder"))
+				.Visibility_Lambda( [] () { return CVarEnableCategoryContentChooserView->GetBool()   ? EVisibility::Visible : EVisibility::Collapsed; })
+				.Padding(0)
+				[
+					CategoryContentBuilder->GenerateWidgetSharedRef()
+				]
+			]
 		+ SVerticalBox::Slot()
 		.AutoHeight()
 		.HAlign(HAlign_Fill)
@@ -662,6 +686,7 @@ void SPlacementModeTools::Construct( const FArguments& InArgs, TSharedRef<SDockT
 		.Padding(FMargin(0.0f, 3.f))
 		[
 			SNew(SOverlay)
+			.Visibility(CVarEnableCategoryContentChooserView->GetBool() ? EVisibility::Collapsed : EVisibility::Visible)
 
 			+ SOverlay::Slot()
 			.HAlign(HAlign_Center)
@@ -737,6 +762,7 @@ void SPlacementModeTools::UpdateShownItems()
 	}
 	else if (Category->CustomGenerator)
 	{
+		UniformWrapPanel->AddSlot()[Category->CustomGenerator()];
 		CustomContent->SetContent(Category->CustomGenerator());
 
 		CustomContent->SetVisibility(EVisibility::Visible);
@@ -768,10 +794,13 @@ void SPlacementModeTools::UpdateShownItems()
 			}
 		}
 
-		CustomContent->SetVisibility(EVisibility::Collapsed);
-		DataDrivenContent->SetVisibility(EVisibility::Visible);
-		ListView->RequestListRefresh();
-		FilterLabelPtr->SetText(Category->DisplayName);
+		if ( !CVarEnableCategoryContentChooserView->GetBool() )
+		{
+			CustomContent->SetVisibility(EVisibility::Collapsed);
+			DataDrivenContent->SetVisibility(EVisibility::Visible);
+			ListView->RequestListRefresh();
+			FilterLabelPtr->SetText(Category->DisplayName);
+		}
 	}
 }
 
@@ -784,7 +813,7 @@ ECheckBoxState SPlacementModeTools::GetPlacementTabCheckedState( FName CategoryN
 {
 	return ActiveTabName == CategoryName ? ECheckBoxState::Checked : ECheckBoxState::Unchecked;
 }
-
+	
 EVisibility SPlacementModeTools::GetFailedSearchVisibility() const
 {
 	if (!IsSearchActive() || FilteredItems.Num())
@@ -796,7 +825,7 @@ EVisibility SPlacementModeTools::GetFailedSearchVisibility() const
 
 EVisibility SPlacementModeTools::GetTabsVisibility() const
 {
-	return IsSearchActive() ? EVisibility::Collapsed : EVisibility::Visible;
+	return IsSearchActive() || CVarEnableCategoryContentChooserView->GetBool() ? EVisibility::Collapsed : EVisibility::Visible;
 }
 
 TSharedRef<ITableRow> SPlacementModeTools::OnGenerateWidgetForItem(TSharedPtr<FPlaceableItem> InItem, const TSharedRef<STableViewBase>& OwnerTable)
@@ -804,11 +833,47 @@ TSharedRef<ITableRow> SPlacementModeTools::OnGenerateWidgetForItem(TSharedPtr<FP
 	return SNew(STableRow<TSharedPtr<FPlaceableItem>>, OwnerTable)
 		.Style(&FAppStyle::Get(), "PlacementBrowser.PlaceableItemRow")
 		[
-			SNew(SPlacementAssetEntry, InItem.ToSharedRef())
-			.HighlightText(this, &SPlacementModeTools::GetHighlightText)
+			GetPlacementAssetWidget(InItem)
 		];
 }
 
+TSharedRef<SWidget> SPlacementModeTools::GetPlacementAssetWidget( const TSharedPtr<FPlaceableItem>& InItem ) const
+{
+	return SNew(SPlacementAssetEntry, InItem.ToSharedRef())
+		.HighlightText(this, &SPlacementModeTools::GetHighlightText);
+}
+
+TSharedRef<SWidget> SPlacementModeTools::ProvideCategoryContent(const FName& CategoryName)
+{
+	if (CVarEnableCategoryContentChooserView->GetBool())
+	{
+		SetActiveTab(CategoryName);
+		UniformWrapPanel->ClearChildren();
+		UpdateShownItems();
+
+		// else, add the items found for the category (and potentially filtered)to our display boxes to show 
+		if ( !FilteredItems.IsEmpty() && UniformWrapPanel->GetChildren()->Num() == 0 )
+		{
+			for (const TSharedPtr<FPlaceableItem>& Item : FilteredItems)
+			{
+				UniformWrapPanel->AddSlot()
+				[
+					GetPlacementAssetWidget(Item)
+				];
+			}
+		}
+
+		return UniformWrapPanel->GetChildren()->Num() == 0 ?
+			SNullWidget::NullWidget :
+			SNew( SScrollBox )
+						+ SScrollBox::Slot()
+						[
+							UniformWrapPanel.ToSharedRef()
+						];
+	}
+	return SNullWidget::NullWidget;
+}
+
 void SPlacementModeTools::OnCategoryChanged(const ECheckBoxState NewState, FName InCategory)
 {
 	if (NewState == ECheckBoxState::Checked)
@@ -860,6 +925,12 @@ void SPlacementModeTools::UpdatePlacementCategories()
 
 	TArray<FPlacementCategoryInfo> Categories;
 	IPlacementModeModule::Get().GetSortedCategories(Categories);
+	
+	// register will make a command to load each Category
+	FPlacementModeCommands::Register();
+	// set the commands as the Categories in our category content chooser
+	CategoryContentBuilder->SetCommands(FPlacementModeCommands::Get().GetPlacementToolkitCategoryCommands());
+	
 	for (const FPlacementCategoryInfo& Category : Categories)
 	{
 		if (Category.UniqueHandle == FBuiltInPlacementCategories::Basic())
@@ -938,6 +1009,9 @@ void SPlacementModeTools::OnSearchChanged(const FText& InFilterText)
 
 	SearchTextFilter->SetRawFilterText( InFilterText );
 	SearchBoxPtr->SetError( SearchTextFilter->GetFilterErrorText() );
+
+	// update the category content which may have changed on search
+	CategoryContentBuilder->UpdateWidget();
 }
 
 void SPlacementModeTools::OnSearchCommitted(const FText& InFilterText, ETextCommit::Type InCommitType)
diff --git a/Engine/Source/Editor/PlacementMode/Private/SPlacementModeTools.h b/Engine/Source/Editor/PlacementMode/Private/SPlacementModeTools.h
--- a/Engine/Source/Editor/PlacementMode/Private/SPlacementModeTools.h
+++ b/Engine/Source/Editor/PlacementMode/Private/SPlacementModeTools.h
@@ -15,6 +15,7 @@
 #include "Misc/TextFilter.h"
 
 class STextBlock;
+class FCategoryDrivenContentBuilder;
 
 /**
  * A tile representation of the class or the asset.  These are embedded into the views inside
@@ -93,6 +94,22 @@ public:
 
 	virtual ~SPlacementModeTools();
 
+	/**
+	 *Creates and returns the TSharedRef<SPlacementAssetEntry>for the given TSharedPtr<SPlacementAssetEntry>
+	 *
+	 * @param InItem the FPlaceableItem for which  the TSharedRef<SPlacementAssetEntry> will be created and returned
+	 */
+	TSharedRef<SWidget> GetPlacementAssetWidget(const TSharedPtr<FPlaceableItem>& InItem) const;
+
+	/**
+	 * Given FName CategoryName as the Category that has been clicked, it should return the
+	 * TSharedRef<SWidget> that contains the content for that Category
+	 * 
+	 * @param CategoryName the name of the Category that the user has clicked
+	 * @return 
+	 */
+	TSharedRef<SWidget> ProvideCategoryContent(const FName& CategoryName);
+
 private:
 
 	// Begin SWidget
@@ -193,4 +210,10 @@ private:
 
 	/** List view that shows placeable items */
 	TSharedPtr<SListView<TSharedPtr<FPlaceableItem>>> ListView;
+
+	/** Wrap panel containing the draggable content  */
+	TSharedPtr<SUniformWrapPanel> UniformWrapPanel;
+	
+	/** The FCategoryDrivenContentBuilder  which will build the UI for this mode, if defined */
+	TSharedPtr<FCategoryDrivenContentBuilder> CategoryContentBuilder;
 };
