Place Actors Panel - Usability Improvements

- fix import in CategoryDrivenContentBuilder.h
- move copyright up in the placement commands

diff --git a/Engine/Source/Developer/WidgetRegistration/Public/Layout/CategoryDrivenContentBuilder.h b/Engine/Source/Developer/WidgetRegistration/Public/Layout/CategoryDrivenContentBuilder.h
--- a/Engine/Source/Developer/WidgetRegistration/Public/Layout/CategoryDrivenContentBuilder.h
+++ b/Engine/Source/Developer/WidgetRegistration/Public/Layout/CategoryDrivenContentBuilder.h
@@ -3,7 +3,7 @@
 #pragma once
 #include "ToolkitBuilder.h"
 #include "ToolkitStyle.h"
-#include "Layout\CategoryDrivenContentBuilderBase.h"
+#include "Layout/CategoryDrivenContentBuilderBase.h"
 
 /**
  * A builder which creates a widget that has a vertical toolbar category picker on the left
