Place Actors Panel - Usability Improvements

- removed the new explicits on the constructors as it is causing compilation errors

diff --git a/Engine/Source/Developer/WidgetRegistration/Public/ToolkitBuilder.h b/Engine/Source/Developer/WidgetRegistration/Public/ToolkitBuilder.h
--- a/Engine/Source/Developer/WidgetRegistration/Public/ToolkitBuilder.h
+++ b/Engine/Source/Developer/WidgetRegistration/Public/ToolkitBuilder.h
@@ -103,7 +103,7 @@ struct FToolkitBuilderArgs;
  */
 struct FToolkitBuilderArgs : public FCategoryDrivenContentBuilderArgs
 {
-	WIDGETREGISTRATION_API explicit FToolkitBuilderArgs(const FName InToolbarCustomizationName);
+	WIDGETREGISTRATION_API FToolkitBuilderArgs(const FName InToolbarCustomizationName);
 
 	/** Name of the toolbar this mode uses and can be used by external systems to customize that mode toolbar */
 	FName ToolbarCustomizationName;
@@ -129,13 +129,13 @@ public:
 	 * @param InToolkitCommandList  the toolkit FUICommandList
 	 * @param InToolkitSections The FToolkitSections for this toolkit builder
 	 */
-	explicit FToolkitBuilder(
+	FToolkitBuilder(
 	FName ToolbarCustomizationName,
 	TSharedPtr<FUICommandList> InToolkitCommandList,
 	TSharedPtr<FToolkitSections> InToolkitSections);
 
 	/** default constructor */
-	explicit FToolkitBuilder(FToolkitBuilderArgs& Args);
+	FToolkitBuilder(FToolkitBuilderArgs& Args);
 
 	virtual ~FToolkitBuilder() override;
 
