Place Actors Panel - Usability Improvements

- added definition for constructor back to toolkit builder

diff --git a/Engine/Source/Developer/WidgetRegistration/Private/ToolkitBuilder.cpp b/Engine/Source/Developer/WidgetRegistration/Private/ToolkitBuilder.cpp
--- a/Engine/Source/Developer/WidgetRegistration/Private/ToolkitBuilder.cpp
+++ b/Engine/Source/Developer/WidgetRegistration/Private/ToolkitBuilder.cpp
@@ -114,6 +114,18 @@ FToolkitBuilder::~FToolkitBuilder()
 	}
 }
 
+FToolkitBuilder::FToolkitBuilder(
+	FName InToolbarCustomizationName,
+	TSharedPtr<FUICommandList> InToolkitCommandList,
+	TSharedPtr<FToolkitSections> InToolkitSections) :
+	FCategoryDrivenContentBuilderBase("FToolkitBuilder"),
+	ToolbarCustomizationName(InToolbarCustomizationName),
+	ToolkitCommandList(InToolkitCommandList),
+	ToolkitSections(InToolkitSections)
+{
+	ResetWidget();
+}
+
 void FToolkitBuilder::GetCommandsForEditablePalette(TSharedRef<FEditablePalette> EditablePalette, TArray<TSharedPtr<const FUICommandInfo>>& OutCommands)
 {
 	TArray<FString> CommandNames = EditablePalette->GetPaletteCommandNames();
