Place actors milestone 2, managers and keys

Added keys for Builders, Icons, Icon sizes, and managers for style, command creation, Inputs, and Persistence.

diff --git a/Engine/Source/Developer/WidgetRegistration/Private/BuilderKey.cpp b/Engine/Source/Developer/WidgetRegistration/Private/BuilderKey.cpp
new file mode 100644
--- /dev/null
+++ b/Engine/Source/Developer/WidgetRegistration/Private/BuilderKey.cpp
@@ -0,0 +1,57 @@
+// Copyright Epic Games, Inc. All Rights Reserved.
+
+#include "BuilderKey.h"
+
+namespace UE::DisplayBuilders::BuilderKey
+{
+	const TCHAR Separator = ',';
+	
+	namespace KeyNames
+	{
+		const FName PlaceActors = "PlaceActors";
+		const FName None = "None";
+	}
+
+}
+
+FName UE::DisplayBuilders::FBuilderKey::ToName() const
+{
+	return Key;
+}
+
+UE::DisplayBuilders::FBuilderKey::FBuilderKey():
+	Key( NAME_None )
+{
+}
+
+FString UE::DisplayBuilders::FBuilderKey::GetKeyWithSuffix(const FName Suffix) const
+{
+	return Key.ToString()  +  UE::DisplayBuilders::BuilderKey::Separator  + Suffix.ToString();
+}
+
+bool UE::DisplayBuilders::FBuilderKey::IsNone() const
+{
+	return Key.IsNone();
+}
+
+const UE::DisplayBuilders::FBuilderKeys& UE::DisplayBuilders::FBuilderKeys::Get()
+{
+	static const FBuilderKeys Keys;
+	return Keys;
+}
+
+const UE::DisplayBuilders::FBuilderKey& UE::DisplayBuilders::FBuilderKeys::PlaceActors() const 
+{
+	static const  UE::DisplayBuilders::FBuilderKey Key{ BuilderKey::KeyNames::PlaceActors };
+	return Key;
+}
+
+const UE::DisplayBuilders::FBuilderKey& UE::DisplayBuilders::FBuilderKeys::None() const 
+{
+	static const  UE::DisplayBuilders::FBuilderKey Key{ BuilderKey::KeyNames::None };
+	return Key;
+}
+
+UE::DisplayBuilders::FBuilderKeys::FBuilderKeys()
+{
+}
diff --git a/Engine/Source/Developer/WidgetRegistration/Private/Common/BuilderTypes.cpp b/Engine/Source/Developer/WidgetRegistration/Private/Common/BuilderTypes.cpp
new file mode 100644
--- /dev/null
+++ b/Engine/Source/Developer/WidgetRegistration/Private/Common/BuilderTypes.cpp
@@ -0,0 +1,9 @@
+// Copyright Epic Games, Inc. All Rights Reserved.
+
+#include "Common/BuilderTypes.h"
+
+UE::DisplayBuilders::FLabelAndIconArgs::FLabelAndIconArgs(FText InLabel, FSlateIcon InIcon):
+	Label( InLabel )
+	, Icon( InIcon )
+{
+}
diff --git a/Engine/Source/Developer/WidgetRegistration/Private/Inputs/BuilderCommandCreationManager.cpp b/Engine/Source/Developer/WidgetRegistration/Private/Inputs/BuilderCommandCreationManager.cpp
new file mode 100644
--- /dev/null
+++ b/Engine/Source/Developer/WidgetRegistration/Private/Inputs/BuilderCommandCreationManager.cpp
@@ -0,0 +1,66 @@
+// Copyright Epic Games, Inc. All Rights Reserved.
+
+#include "Inputs/BuilderCommandCreationManager.h"
+#include "Framework/Commands/UICommandInfo.h"
+#include "Inputs/BuilderInput.h"
+
+#define LOCTEXT_NAMESPACE "BuilderCommandCreationManager"
+
+FBuilderCommandCreationManager::FBuilderCommandCreationManager()
+	: TCommands<FBuilderCommandCreationManager>(
+		TEXT("BuilderCommandCreationManager"), 
+		LOCTEXT("BuilderCommandCreationManager", "Dynamic FUI Commands"), 
+		NAME_None, 
+		"FBuilderCommandCreationManagerStyle")
+{
+}
+
+const FBuilderCommandCreationManager& FBuilderCommandCreationManager::Get()
+{
+	return TCommands<FBuilderCommandCreationManager>::Get();
+}
+
+void FBuilderCommandCreationManager::RegisterCommands()
+{
+	/*
+	 * we must put in one command for the singleton to stay valid after it is registered. Since this is for dynamic commands we register a default
+	 */
+	static const FName Default{ "Default" };
+	static UE::DisplayBuilders::FBuilderInput Input{}; 
+	Input.Name = Default;
+	Input.Label = FText::FromName(Default);
+	RegisterCommandForBuilder( Input );
+}
+
+void FBuilderCommandCreationManager::RegisterCommandForBuilder( UE::DisplayBuilders::FBuilderInput& Input ) const
+{
+	if (Instance.IsValid() && !Input.IsNameNone() )
+	{
+		const TSharedPtr<FBuilderCommandCreationManager> Commands = Instance.Pin();
+
+		TSharedPtr<FUICommandInfo> NewCommandInfo;
+
+		FUICommandInfo::MakeCommandInfo(
+			Commands->AsShared(),
+			NewCommandInfo,
+			Input.Name,
+			Input.Label,
+			Input.Tooltip,
+			Input.Icon,
+			EUserInterfaceActionType::RadioButton,
+			Input.DefaultChords );
+
+		Input.UICommandInfo = NewCommandInfo;
+		Input.ButtonArgs.Command = NewCommandInfo;
+	};
+}
+
+void FBuilderCommandCreationManager::UnregisterCommandForBuilder( UE::DisplayBuilders::FBuilderInput& Input ) const
+{
+	if ( Input.UICommandInfo.IsValid() )
+	{
+		FUICommandInfo::UnregisterCommandInfo( Instance.Pin().ToSharedRef(), Input.UICommandInfo.ToSharedRef() );		
+	}
+}
+
+#undef LOCTEXT_NAMESPACE
diff --git a/Engine/Source/Developer/WidgetRegistration/Private/Inputs/BuilderCommandCreationManager.h b/Engine/Source/Developer/WidgetRegistration/Private/Inputs/BuilderCommandCreationManager.h
new file mode 100644
--- /dev/null
+++ b/Engine/Source/Developer/WidgetRegistration/Private/Inputs/BuilderCommandCreationManager.h
@@ -0,0 +1,48 @@
+// Copyright Epic Games, Inc. All Rights Reserved.
+
+#pragma once
+
+#include "Framework/Commands/Commands.h"
+#include "Inputs/BuilderInput.h"
+
+namespace UE::DisplayBuilders
+{
+	class FBuilderInput;
+}
+
+/**
+* FBuilderCommandCreationManager creates dynamic FUICommandInfos that can be used with toolbars, menus, and other command-centric objects and
+* fills out the associated information in FBuilderInputs
+*/
+class FBuilderCommandCreationManager : public TCommands<FBuilderCommandCreationManager>
+{
+	
+public:
+	/** default constructor  */
+	FBuilderCommandCreationManager();
+
+	/**
+	 * Creates the current commands specified in FUICommandInfoArray
+	 */
+	virtual void RegisterCommands() override;
+
+	/**
+   	* Given a FBuilderInput, registers a TSharedPtr<FUICommandInfo> which can be used to create Slate 
+   	* toolbars, menus, and other command-centric objects and initializes the Command information for the FBuilderInput 
+   	*
+   	* @param OutBuilderInput the FBuilderInput for which to register and initialize the related TSharedPtr<FUICommandInfo>
+   	*/
+	void RegisterCommandForBuilder( UE::DisplayBuilders::FBuilderInput& OutBuilderInput ) const;
+
+	/**
+    * Given a FBuilderInput, unregisters and tears down the related TSharedPtr<FUICommandInfo> 
+    *
+    * @param OutBuilderInput the FBuilderInput for which to unregister and tear down TSharedPtr<FUICommandInfo>
+    */
+	void UnregisterCommandForBuilder( UE::DisplayBuilders::FBuilderInput& OutBuilderInput ) const;
+
+	/**
+	 * returns the singleton of the FBuilderCommandCreationManager
+	 */
+	static const FBuilderCommandCreationManager& Get();
+};
\ No newline at end of file
diff --git a/Engine/Source/Developer/WidgetRegistration/Private/Inputs/BuilderInputManager.cpp b/Engine/Source/Developer/WidgetRegistration/Private/Inputs/BuilderInputManager.cpp
new file mode 100644
--- /dev/null
+++ b/Engine/Source/Developer/WidgetRegistration/Private/Inputs/BuilderInputManager.cpp
@@ -0,0 +1,34 @@
+// Copyright Epic Games, Inc. All Rights Reserved.
+
+#include "Inputs/BuilderInputManager.h"
+#include "BuilderCommandCreationManager.h"
+#include "Framework/Commands/UICommandInfo.h"
+// #include "Inputs/DragAndDrop/BuilderDragAndDropManager.h"
+
+FBuilderInputManager& FBuilderInputManager::Get()
+{
+	static FBuilderInputManager Manager;
+	return Manager;
+}
+
+void FBuilderInputManager::Initialize()
+{
+	FBuilderCommandCreationManager::Register();
+}
+
+void FBuilderInputManager::Shutdown()
+{
+	FBuilderCommandCreationManager::Unregister();
+}
+
+const FBuilderCommandCreationManager& FBuilderInputManager::GetCommandManager()
+{
+	static const FBuilderCommandCreationManager& CommandManager = FBuilderCommandCreationManager::Get();
+	return CommandManager;
+}
+
+/*FBuilderDragAndDropManager& FBuilderInputManager::GetDragAndDropManager()
+{
+	static FBuilderDragAndDropManager& DragAndDropManager = FBuilderDragAndDropManager::Get();
+	return DragAndDropManager;
+}*/
\ No newline at end of file
diff --git a/Engine/Source/Developer/WidgetRegistration/Private/Inputs/BuilderInputManager.h b/Engine/Source/Developer/WidgetRegistration/Private/Inputs/BuilderInputManager.h
new file mode 100644
--- /dev/null
+++ b/Engine/Source/Developer/WidgetRegistration/Private/Inputs/BuilderInputManager.h
@@ -0,0 +1,46 @@
+// Copyright Epic Games, Inc. All Rights Reserved.
+
+#pragma once
+
+#include "Framework/Commands/UICommandInfo.h"
+
+// class FBuilderDragAndDropManager;
+class FBuilderCommandCreationManager;
+
+namespace UE::DisplayBuilders
+{
+	class FBuilderInput;
+}
+
+/**
+* BuilderInputManager provides access to utilities and information regarding user input.
+*/
+class FBuilderInputManager
+{
+public:
+
+	/**
+	* Gets the singleton instance of FBuilderInputManager
+	*/
+	static FBuilderInputManager& Get();
+
+	/**
+	 * Registers and initializes the Command binding context.
+	 */
+	static void Initialize();
+
+	/**
+	 * Unregisters the Command context.
+	 */
+	static void Shutdown();
+
+	/**
+	 * Returns the Command manager for Builders
+	 */
+	const FBuilderCommandCreationManager& GetCommandManager();
+	
+	/**
+	 * Returns the Drag and Drop manager for Builders
+	 */
+	// FBuilderDragAndDropManager& GetDragAndDropManager();
+};
diff --git a/Engine/Source/Developer/WidgetRegistration/Private/Persistence/BuilderPersistenceManager.cpp b/Engine/Source/Developer/WidgetRegistration/Private/Persistence/BuilderPersistenceManager.cpp
new file mode 100644
--- /dev/null
+++ b/Engine/Source/Developer/WidgetRegistration/Private/Persistence/BuilderPersistenceManager.cpp
@@ -0,0 +1,59 @@
+// Copyright Epic Games, Inc. All Rights Reserved.
+
+#include "Persistence/BuilderPersistenceManager.h"
+
+TObjectPtr<UBuilderPersistenceManager> UBuilderPersistenceManager::Instance = nullptr;
+
+FName UBuilderPersistenceManager::FavoritesSuffix = "BuilderFavorites";
+
+void UBuilderPersistenceManager::Initialize()
+{
+	if(!Instance)
+	{
+		Instance = NewObject<UBuilderPersistenceManager>(); 
+		Instance->AddToRoot();
+		Instance->LoadEditorConfig();
+	}
+}
+
+void UBuilderPersistenceManager::ShutDown()
+{
+	Instance->RemoveFromRoot();
+	Instance = nullptr;
+}
+
+TArray<FName> UBuilderPersistenceManager::GetPersistedFavoritesNamesArray(const UE::DisplayBuilders::FBuilderKey& Key)
+{
+	return GetPersistedArrayOfNames(Key, FavoritesSuffix);
+}
+
+void UBuilderPersistenceManager::SetPersistedFavoritesNamesArray(const UE::DisplayBuilders::FBuilderKey& Key, TArray<FName>& Favorites)
+{
+	return PersistArrayOfNames(Key, FavoritesSuffix, Favorites);
+}
+
+TArray<FName> UBuilderPersistenceManager::GetPersistedArrayOfNames( const UE::DisplayBuilders::FBuilderKey& Key, FName PersistenceKeySuffix )
+{
+	if ( !Key.IsNone()  && !PersistenceKeySuffix.IsNone() )
+	{
+		if ( const FPersistedNameArray* Settings =  PersistenceSettings.Find( Key.GetKeyWithSuffix( PersistenceKeySuffix ) ) )
+		{
+			return Settings->ArrayOfNamesToPersist;
+		}
+	}
+
+	return {};
+}
+
+void UBuilderPersistenceManager::PersistArrayOfNames( const UE::DisplayBuilders::FBuilderKey& Key, FName PersistenceKeySuffix, TArray<FName>& ArrayOfNamesToPersist )
+{
+	if (PersistenceKeySuffix.IsNone())
+	{
+		return;
+	}
+
+	FPersistedNameArray& Settings =  PersistenceSettings.Add( Key.GetKeyWithSuffix( PersistenceKeySuffix ) );
+
+	Settings.ArrayOfNamesToPersist = ArrayOfNamesToPersist;;
+	SaveEditorConfig();
+}
\ No newline at end of file
diff --git a/Engine/Source/Developer/WidgetRegistration/Private/Styles/BuilderIconKeys.cpp b/Engine/Source/Developer/WidgetRegistration/Private/Styles/BuilderIconKeys.cpp
new file mode 100644
--- /dev/null
+++ b/Engine/Source/Developer/WidgetRegistration/Private/Styles/BuilderIconKeys.cpp
@@ -0,0 +1,56 @@
+// Copyright Epic Games, Inc. All Rights Reserved.
+
+#include "BuilderIconKeys.h"
+
+#include "BuilderIconSizeKeys.h"
+#include "BuilderStyleManager.h"
+
+// folder path information
+namespace UE::DisplayBuilders::BuilderIconsKeys::Paths
+{
+	const TCHAR Slash = '/';
+	const TCHAR* Icons_DataVisualization = TEXT("Icons/DataVisualization");
+}
+
+// icon name information
+namespace UE::DisplayBuilders::BuilderIconsKeys::Icons
+{
+	const FName ZeroStateDefault = "ZeroStateDefault";
+	const FName ZeroStateFavorites = "ZeroStateFavorites";
+}
+
+const FBuilderIconKeys& FBuilderIconKeys::Get()
+{
+	static const FBuilderIconKeys Keys;
+	return Keys;
+}
+
+const FBuilderIconKey& FBuilderIconKeys::ZeroStateDefaultMedium() const
+{
+	static const FBuilderIconKey Key( UE::DisplayBuilders::BuilderIconsKeys::Paths::Icons_DataVisualization,UE::DisplayBuilders::BuilderIconsKeys:: Icons::ZeroStateDefault, FBuilderIconSizeKeys::Get().Medium() );
+	return Key;
+}
+
+const FBuilderIconKey& FBuilderIconKeys::ZeroStateFavoritesMedium() const
+{
+	static const FBuilderIconKey Key( UE::DisplayBuilders::BuilderIconsKeys::Paths::Icons_DataVisualization,UE::DisplayBuilders::BuilderIconsKeys:: Icons::ZeroStateFavorites, FBuilderIconSizeKeys::Get().Medium() );
+	return Key;
+}
+
+FSlateIcon FBuilderIconKey::GetSlateIcon() const
+{
+	return  FSlateIcon( FBuilderStyleManager::Get().GetStyleSetName(), FileNameWithoutExtension );
+}
+
+FBuilderIconKey::FBuilderIconKey( FString InRelativePathToContainingFolder, const FName InName, const FBuilderIconSizeKey& InSizeKey):
+	RelativePathToContainingFolder( InRelativePathToContainingFolder )
+	, Name( InName )
+	, SizeKey( InSizeKey )
+	, FileNameWithoutExtension( Name.ToString() + InSizeKey.Name.ToString() )
+	, RelativePathToFileWithoutExtension( RelativePathToContainingFolder + UE::DisplayBuilders::BuilderIconsKeys::Paths::Slash + FileNameWithoutExtension.ToString() )
+{
+}
+
+FBuilderIconKeys::FBuilderIconKeys() 
+{
+}
diff --git a/Engine/Source/Developer/WidgetRegistration/Private/Styles/BuilderIconKeys.h b/Engine/Source/Developer/WidgetRegistration/Private/Styles/BuilderIconKeys.h
new file mode 100644
--- /dev/null
+++ b/Engine/Source/Developer/WidgetRegistration/Private/Styles/BuilderIconKeys.h
@@ -0,0 +1,94 @@
+// Copyright Epic Games, Inc. All Rights Reserved.
+
+#pragma once
+
+#include "BuilderIconSizeKeys.h"
+#include "BuilderStyleManager.h"
+#include "Textures/SlateIcon.h"
+
+class FBuilderIconSizeKeys;
+class FBuilderIconSizeKey;
+class FBuilderIconKeys;
+
+/**
+* Provides keys for display builder Icons 
+*/
+class FBuilderIconKey
+{
+public:
+
+	/** initializes the slate icon. This must be called after the the icon is registered by the Builder Style */
+	void InitializeIcon()
+	{
+		SlateIcon =  FSlateIcon( FBuilderStyleManager::Get().GetStyleSetName(), FileNameWithoutExtension );
+	}
+
+	/**
+	* The FString providing the relative path to the icon 
+	*/
+	const FString RelativePathToContainingFolder;
+
+	/**
+	* The name of the Icon 
+	*/
+	const FName Name;
+	
+	/**
+	 * The FBuilderIconSizeKey that provides the size of the Icon
+	 */
+	const FBuilderIconSizeKey& SizeKey;
+
+	/**
+	* The name of the Icon file, not including relative path or extension
+	*/
+	const FName FileNameWithoutExtension;
+
+	/**
+	 * the relative path to the file for the icon. This includes both the relative path of directories and the name of the file for the icon,
+	 * but not the file extension
+	 */
+	const FString RelativePathToFileWithoutExtension;
+	
+
+	/**
+	 * returns the FSlateIcon for this FBuilderIconKey
+	 */
+	FSlateIcon GetSlateIcon() const;
+
+private:
+
+	friend FBuilderIconKeys;
+	
+	FSlateIcon SlateIcon;
+	/**
+	* The constructor which takes an FName which provides the identifier for this BuilderIconKey
+	*
+	* @param InName the FName which provides the identifier for this BuilderIconKey
+	* @param InRelativePathToContainingFolder the relative path to the containing folder
+	* @param InSizeKey the icon size key for this 
+	*/
+	 FBuilderIconKey( const FString InRelativePathToContainingFolder, const FName InName, const FBuilderIconSizeKey& InSizeKey );
+};
+
+
+/**
+ * Keys for builder icons
+ */
+class FBuilderIconKeys
+{
+public:
+	/**
+	 * Gets the FBuilderIconKeys singleton
+	 */
+	static const FBuilderIconKeys& Get();
+
+	/** An Icon for the default Zero State view. This shows a default icon that denotes that there are no items available */
+	const FBuilderIconKey& ZeroStateDefaultMedium() const;
+	
+	/** An Icon for the favorites Zero State view. This shows a zero state favorites icon that denotes that there are no favorites available */
+	const FBuilderIconKey& ZeroStateFavoritesMedium() const;
+	
+private:
+	/** private constructor as these keys should map to icons that we include with builder content */
+	FBuilderIconKeys();
+};
diff --git a/Engine/Source/Developer/WidgetRegistration/Private/Styles/BuilderIconSizeKeys.h b/Engine/Source/Developer/WidgetRegistration/Private/Styles/BuilderIconSizeKeys.h
new file mode 100644
--- /dev/null
+++ b/Engine/Source/Developer/WidgetRegistration/Private/Styles/BuilderIconSizeKeys.h
@@ -0,0 +1,88 @@
+// Copyright Epic Games, Inc. All Rights Reserved.
+
+#pragma once
+
+#include "Containers/UnrealString.h"
+#include "UObject/NameTypes.h"
+#include "Math/Vector2D.h"
+
+class FBuilderIconSizeKeys;
+
+/**
+* Provides identifiers and values for icon sizes 
+*/
+class FBuilderIconSizeKey
+{
+public:
+	/**
+	 * default constructor that initializes the key
+	 */
+	explicit FBuilderIconSizeKey( FString RelativePathToContainingFolder );
+
+	/**
+	* The name of the Icon 
+	*/
+	const FName Name;
+
+	/**
+	 * The Size of the Icon
+	 */
+	const FVector2D Size;
+ 	
+	
+private:
+
+	friend FBuilderIconSizeKeys;
+	
+	/**
+	* The constructor which takes an FName which provides the identifier for this BuilderIconSizeKey
+	*
+	* @param InName the FName which provides the identifier for this BuilderIconSizeKey
+	* @param InSize the FVector2D size of the icon
+	*/
+	 FBuilderIconSizeKey( const FName InName, const FVector2D InSize ) :
+		Name( InName )
+		, Size( InSize )
+	{
+	}
+};
+
+/**
+ * FBuilderIconSizeKeys provides sizes for builder icons
+ */
+class FBuilderIconSizeKeys
+{
+public:
+	static const FBuilderIconSizeKeys& Get();
+
+	const FBuilderIconSizeKey& ExtraSmall() const
+	{
+		static const FBuilderIconSizeKey Key( "ExtraSmall", FVector2D(16, 16) );
+		return Key;
+	}
+
+	const FBuilderIconSizeKey& Small() const
+	{
+		static const FBuilderIconSizeKey Key( "ExtraSmall", FVector2D(20, 20) );
+		return Key;
+	}
+
+	const FBuilderIconSizeKey& Medium() const
+	{
+		static const FBuilderIconSizeKey Key( "ExtraSmall", FVector2D(40, 40) );
+		return Key;
+	}
+	
+private:
+	FBuilderIconSizeKeys();
+};
+
+inline const FBuilderIconSizeKeys& FBuilderIconSizeKeys::Get()
+{
+	static const FBuilderIconSizeKeys Keys;
+	return Keys;
+}
+
+inline FBuilderIconSizeKeys::FBuilderIconSizeKeys()
+{
+}
diff --git a/Engine/Source/Developer/WidgetRegistration/Private/Styles/BuilderStyleManager.cpp b/Engine/Source/Developer/WidgetRegistration/Private/Styles/BuilderStyleManager.cpp
new file mode 100644
--- /dev/null
+++ b/Engine/Source/Developer/WidgetRegistration/Private/Styles/BuilderStyleManager.cpp
@@ -0,0 +1,32 @@
+// Copyright Epic Games, Inc. All Rights Reserved.
+
+#include "BuilderStyleManager.h"
+
+#include "BuilderIconKeys.h"
+#include "Brushes/SlateImageBrush.h"
+#include "Misc/Paths.h"
+#include "Styling/SlateStyleMacros.h"
+#include "Styling/SlateStyleRegistry.h"
+
+void FBuilderStyleManager::RegisterSlateIcon( FBuilderIconKey Key )
+{
+	Set( Key.FileNameWithoutExtension, new IMAGE_BRUSH_SVG( Key.RelativePathToFileWithoutExtension, Key.SizeKey.Size ));
+	Key.InitializeIcon();
+}
+
+FBuilderStyleManager::FBuilderStyleManager()
+	: FSlateStyleSet(TEXT("Builder"))
+{
+
+	SetContentRoot(FPaths::EngineContentDir() / TEXT("Editor/Slate/Builders"));
+	SetCoreContentRoot(FPaths::EngineContentDir() / TEXT("Slate"));
+
+	FText Text = FText::GetEmpty();
+	
+	FSlateStyleRegistry::RegisterSlateStyle(*this);
+}
+
+FBuilderStyleManager::~FBuilderStyleManager()
+{
+	FSlateStyleRegistry::UnRegisterSlateStyle(*this);
+}
diff --git a/Engine/Source/Developer/WidgetRegistration/Private/Styles/BuilderStyleManager.h b/Engine/Source/Developer/WidgetRegistration/Private/Styles/BuilderStyleManager.h
new file mode 100644
--- /dev/null
+++ b/Engine/Source/Developer/WidgetRegistration/Private/Styles/BuilderStyleManager.h
@@ -0,0 +1,31 @@
+// Copyright Epic Games, Inc. All Rights Reserved.
+
+#pragma once
+
+#include "Styling/SlateStyle.h"
+
+class FBuilderIconKey;
+
+/**
+ * Provides styling utilities and information for builders.
+ */
+class FBuilderStyleManager final : public FSlateStyleSet
+{
+public:
+	static FBuilderStyleManager& Get()
+	{
+		static FBuilderStyleManager Instance;
+		return Instance;
+	}
+
+	/**
+	 * Given Key, registers the FSlateIcon for it
+	 * 
+	 * @param Key the FBuilderIconKey to register the FSlateIcon for
+	 */
+	void RegisterSlateIcon( const FBuilderIconKey Key );
+	
+	FBuilderStyleManager();
+	
+	virtual ~FBuilderStyleManager() override;
+};
diff --git a/Engine/Source/Developer/WidgetRegistration/Private/Styles/SlateBrushTemplates.cpp b/Engine/Source/Developer/WidgetRegistration/Private/Styles/SlateBrushTemplates.cpp
--- a/Engine/Source/Developer/WidgetRegistration/Private/Styles/SlateBrushTemplates.cpp
+++ b/Engine/Source/Developer/WidgetRegistration/Private/Styles/SlateBrushTemplates.cpp
@@ -1,8 +1,56 @@
 // Copyright Epic Games, Inc. All Rights Reserved.
 
 #include "Styles/SlateBrushTemplates.h"
+#include "Brushes/SlateColorBrush.h"
+#include "Styling/CoreStyle.h"
+#include "Styling/StyleColors.h"
+
+const FSlateBrush* FSlateBrushTemplates::DragHandle()
+{
+	const FSlateBrush* Brush = FCoreStyle::Get().GetBrush("VerticalBoxDragIndicatorShort");
+	return Brush;
+}
 
 const FSlateBrush* FSlateBrushTemplates::ThinLineHorizontal()
 {
 	return FAppStyle::GetBrush( "ThinLine.Horizontal");		
 }
+
+const FSlateBrush* FSlateBrushTemplates::Transparent()
+{
+	static const FSlateColorBrush ColorBrush{ FLinearColor::Transparent };
+	return &ColorBrush;
+}
+
+const FSlateBrush* FSlateBrushTemplates::Panel()
+{
+	static const FSlateColorBrush ColorBrush{ FStyleColors::Panel };
+	return &ColorBrush;
+}
+
+const FSlateBrush* FSlateBrushTemplates::Recessed()
+{
+	static const FSlateColorBrush ColorBrush{ FStyleColors::Recessed };
+	return &ColorBrush;
+}
+
+const FSlateBrush* FSlateBrushTemplates::GetBrushWithColor(EStyleColor Color)
+{
+	const FSlateBrush* ColorBrush = EStyleColorToSlateBrushMap.Find( Color );
+	
+	if ( ColorBrush )
+	{
+		return ColorBrush;
+	}
+
+	const FSlateColorBrush NewColorBrush{ Color };
+	EStyleColorToSlateBrushMap.Add( Color, static_cast<const FSlateBrush>( NewColorBrush ) );
+
+	return EStyleColorToSlateBrushMap.Find( Color );;
+}
+
+FSlateBrushTemplates& FSlateBrushTemplates::Get()
+{
+	static FSlateBrushTemplates Templates;
+	return Templates;
+}
diff --git a/Engine/Source/Developer/WidgetRegistration/Private/ToolElementRegistry.cpp b/Engine/Source/Developer/WidgetRegistration/Private/ToolElementRegistry.cpp
--- a/Engine/Source/Developer/WidgetRegistration/Private/ToolElementRegistry.cpp
+++ b/Engine/Source/Developer/WidgetRegistration/Private/ToolElementRegistry.cpp
@@ -26,7 +26,14 @@ FToolElementRegistrationArgs::FToolElementRegistrationArgs( FName InStyleClassNa
 	StyleClassName(InStyleClassName)
 {
 }
-	
+
+FToolElementRegistrationArgs::FToolElementRegistrationArgs( UE::DisplayBuilders::FBuilderKey InBuilderKey ) :
+	ToolElementType(EToolElement::Section)
+	, StyleClassName( InBuilderKey.ToName() )
+	, BuilderKey( InBuilderKey )
+{
+}
+
 TSharedPtr<SWidget> FToolElementRegistrationArgs::GenerateWidget()
 {
 	return nullptr;
diff --git a/Engine/Source/Developer/WidgetRegistration/Private/WidgetRegistrationModule.cpp b/Engine/Source/Developer/WidgetRegistration/Private/WidgetRegistrationModule.cpp
--- a/Engine/Source/Developer/WidgetRegistration/Private/WidgetRegistrationModule.cpp
+++ b/Engine/Source/Developer/WidgetRegistration/Private/WidgetRegistrationModule.cpp
@@ -2,16 +2,22 @@
 
 #include "WidgetRegistrationModule.h"
 #include "ToolkitStyle.h"
+#include "Inputs/BuilderInputManager.h"
+#include "Persistence/BuilderPersistenceManager.h"
 
 
 void FWidgetRegistrationModule::StartupModule()
 {
 	FToolkitStyle::Initialize();
+	UBuilderPersistenceManager::Initialize();
+	FBuilderInputManager::Initialize();
 }
 	
 void FWidgetRegistrationModule::ShutdownModule()
 {
 	FToolkitStyle::Shutdown();
+	UBuilderPersistenceManager::ShutDown();
+	FBuilderInputManager::Shutdown();
 }
 
 
diff --git a/Engine/Source/Developer/WidgetRegistration/Public/BuilderKey.h b/Engine/Source/Developer/WidgetRegistration/Public/BuilderKey.h
new file mode 100644
--- /dev/null
+++ b/Engine/Source/Developer/WidgetRegistration/Public/BuilderKey.h
@@ -0,0 +1,82 @@
+// Copyright Epic Games, Inc. All Rights Reserved.
+
+#pragma once
+
+#include "UObject/NameTypes.h"
+
+namespace UE::DisplayBuilders
+{
+	class FBuilderKeys;
+	
+	/**
+	* Provides identifiers for display builders and some methods for using them. These are also used as the keys into persistent storage for builders.
+	*/
+	class FBuilderKey
+	{ 
+	public:
+		/**
+		 * default constructor that initializes the key
+		 */
+		WIDGETREGISTRATION_API explicit FBuilderKey();
+		
+		/**
+		* converts the BuilderKey to its FName
+		*/
+		WIDGETREGISTRATION_API FName ToName() const;
+
+		/** 
+		 * Returns a combined key with Suffix
+		 * 
+		 * @param Suffix the FName that will provide the Suffix for the combined Key
+		 */
+		FString  GetKeyWithSuffix( const FName Suffix ) const;
+
+		/**
+		* returns true if this builder is set to None, meaning that it has not been initialized
+		*/
+		bool IsNone() const;
+		
+	private:
+
+		friend FBuilderKeys;
+
+		/**
+		* The FName providing the identifier 
+		*/
+		const FName Key;
+ 		
+		/**
+		* The constructor which takes an FName which provides the identifier for this BuilderKey
+		*
+		* @param InKey the FName which provides the identifier for this BuilderKey
+		*/
+		FBuilderKey( const FName InKey ) :
+			Key( InKey )	
+		{
+		}
+	};
+
+	/**
+	 * FBuilderKeys provides keys registered for specific builders. 
+	 */
+	class FBuilderKeys
+	{
+	public:
+	/**
+	* Get the singleton FBuilderKeys 
+	*/
+		WIDGETREGISTRATION_API static const FBuilderKeys& Get();
+
+		/** The FBuilderKey for the Place Actors FCategoryDrivenContentBuilder. This provides the key into the favorites storage for place actors. */
+		WIDGETREGISTRATION_API const FBuilderKey& PlaceActors()  const;
+
+		/** A "None" key to provide a null state when needed. */
+		WIDGETREGISTRATION_API const FBuilderKey& None()  const;
+
+	private:
+		/**
+	 * The constructor is kept private because these are for UE supported builders. 
+	 */
+		FBuilderKeys();
+	};
+}
diff --git a/Engine/Source/Developer/WidgetRegistration/Public/Common/BuilderTypes.h b/Engine/Source/Developer/WidgetRegistration/Public/Common/BuilderTypes.h
new file mode 100644
--- /dev/null
+++ b/Engine/Source/Developer/WidgetRegistration/Public/Common/BuilderTypes.h
@@ -0,0 +1,21 @@
+// Copyright Epic Games, Inc. All Rights Reserved.
+
+#pragma once
+
+#include "Internationalization/Text.h"
+#include "Textures/SlateIcon.h"
+
+namespace UE::DisplayBuilders
+{
+	/**
+	 * A class which provides a parameter object for builders consisting of an FText Label and a FSlateIcon icon
+	 */
+	class FLabelAndIconArgs
+	{
+	public:
+		WIDGETREGISTRATION_API  explicit FLabelAndIconArgs( FText InLabel = FText::GetEmpty(), FSlateIcon InIcon = FSlateIcon() );
+		
+		FText Label;
+		FSlateIcon Icon;
+	};
+}
\ No newline at end of file
diff --git a/Engine/Source/Developer/WidgetRegistration/Public/Inputs/BuilderInput.cpp b/Engine/Source/Developer/WidgetRegistration/Public/Inputs/BuilderInput.cpp
new file mode 100644
--- /dev/null
+++ b/Engine/Source/Developer/WidgetRegistration/Public/Inputs/BuilderInput.cpp
@@ -0,0 +1,58 @@
+// Copyright Epic Games, Inc. All Rights Reserved.
+
+#include "Inputs/BuilderInput.h"
+#include "Inputs/BuilderCommandCreationManager.h"
+
+UE::DisplayBuilders::FBuilderInput UE::DisplayBuilders::FBuilderInput::NullInput;
+
+UE::DisplayBuilders::FBuilderInput::~FBuilderInput()
+{
+	if ( !IsNameNone() )
+	{
+		FBuilderCommandCreationManager::Get().UnregisterCommandForBuilder( *this );		
+	}
+}
+
+bool UE::DisplayBuilders::FBuilderInput::IsNameNone() const
+{
+	return Name.IsNone();
+}
+
+void UE::DisplayBuilders::FBuilderInput::InitializeCommandInfo()
+{
+	if ( !IsNameNone() )
+	{
+		FBuilderCommandCreationManager::Get().RegisterCommandForBuilder( *this );		
+	}
+}
+
+UE::DisplayBuilders::FBuilderInput::FBuilderInput(
+	FName InName,
+	FText InLabel,
+	FSlateIcon InIcon,
+	EUserInterfaceActionType InUserInterfaceType,
+	FText InToolTip,
+	FText InDescription,
+	TArray<TSharedRef<FInputChord>> InActiveChords ,
+	FInputChord InDefaultChords ,
+	FName InUiStyle ,
+	FName InBindingContext  ,
+	FName InBundle  ,
+	bool bInUseLongDisplayName 
+	/*TSharedPtr<FDropZoneHandler> InDropZoneHandler*/   )  :
+	FLabelAndIconArgs( InLabel.IsEmpty() ? FText::FromName( InName ) : InLabel, InIcon )
+	, Name( InName )
+	, UserInterfaceType( InUserInterfaceType )
+	, Description( InDescription.IsEmpty() ? Label : InDescription )
+	, ActiveChords( InActiveChords )
+	, DefaultChords( InDefaultChords )
+	, UIStyle( InUiStyle )
+	, BindingContext( InBindingContext )
+	, Bundle( InBundle )
+	, bUseLongDisplayName( bInUseLongDisplayName )
+	//, DropZoneHandler( InDropZoneHandler )
+	, Index( INDEX_NONE )
+	, Tooltip( InToolTip.IsEmpty() ? InLabel : InToolTip )
+{
+	InitializeCommandInfo();
+}
\ No newline at end of file
diff --git a/Engine/Source/Developer/WidgetRegistration/Public/Inputs/BuilderInput.h b/Engine/Source/Developer/WidgetRegistration/Public/Inputs/BuilderInput.h
new file mode 100644
--- /dev/null
+++ b/Engine/Source/Developer/WidgetRegistration/Public/Inputs/BuilderInput.h
@@ -0,0 +1,101 @@
+// Copyright Epic Games, Inc. All Rights Reserved.
+
+#pragma once
+
+#include "Common/BuilderTypes.h"
+#include "Templates/SharedPointer.h"
+#include "Internationalization/Text.h"
+#include "Textures/SlateIcon.h"
+#include "Framework/Commands/UICommandInfo.h"
+#include "Framework/MultiBox/SToolBarButtonBlock.h"
+#include "Misc/CoreMiscDefines.h"
+
+class FBuilderInputManager;
+struct FInputChord;
+
+namespace UE::DisplayBuilders
+{
+	/**
+	 * BuilderInput is the input representation for Builders and provides conversions to the various Slate types for inputs.
+	 */
+	class FBuilderInput : public FLabelAndIconArgs
+	{
+	public:
+		
+		WIDGETREGISTRATION_API explicit FBuilderInput(
+			FName InName = FName(),
+			FText InLabel = FText::GetEmpty(),
+			FSlateIcon InIcon = FSlateIcon(),
+			EUserInterfaceActionType InUserInterfaceType = EUserInterfaceActionType::Button,
+			FText InToolTip = FText::GetEmpty(),
+			FText InDescription = FText::GetEmpty(),
+			TArray<TSharedRef<FInputChord>> InActiveChords = {},
+			FInputChord InDefaultChords = FInputChord(),
+			FName InUiStyle = NAME_None,
+			FName InBindingContext = NAME_None,
+			FName InBundle = NAME_None,
+			bool bInUseLongDisplayName = false
+			/*TSharedPtr<FDropZoneHandler> InDropZoneHandler = nullptr*/ );
+
+		/**
+		 * The destructor, destroys any command related information
+		 */
+		WIDGETREGISTRATION_API ~FBuilderInput();
+
+		/** Name of the Input */
+		FName Name;
+
+		/** The type of user interface to associated with this action */
+		EUserInterfaceActionType UserInterfaceType;
+
+		/** Localized help text for the UI command */
+		FText Description;
+		
+		/** Input commands that executes this action */
+		TArray<TSharedRef<FInputChord>> ActiveChords;
+		
+		/** The default input chords for the UI command (can be invalid) */
+		FInputChord DefaultChords;
+
+		/** Brush name for icon to use in tool bars and menu items to represent the UI command in its toggled on (checked) state*/
+		FName UIStyle;
+
+		/** The context in which the UI command is active */
+		FName BindingContext;
+
+		/** The bundle to group the UI command into. The bundle must have been added to the BindingContext first. */
+		FName Bundle;
+
+		/** True if should we use long names for when getting text for input chords */
+		bool bUseLongDisplayName;
+
+		/** A pointer to the handler for a dropzone on the input, if one exists */
+		// TSharedPtr<FDropZoneHandler> DropZoneHandler;
+
+		/** the index of the BuilderInput in whatever container it is in */
+		int32 Index;
+		
+		/** the tooltip for the input */
+		FText Tooltip;
+
+		/** FUICommandInfo, the basic Slate Class for defining user-facing commands, including their appearance and behavior in the UI  */
+		TSharedPtr<FUICommandInfo> UICommandInfo;
+
+		/** FButtonArgs, a type which can be used to create buttons with FToolbarBuilder */
+		FButtonArgs ButtonArgs;
+		
+		/** For use when we need a null state */
+		static FBuilderInput NullInput;
+		
+		/**
+		 * returns true if the input name is equivalent to Name_NONE
+		 */
+		bool IsNameNone() const;
+
+	private:
+		/**
+		 * Initializes the related FUICommandInfo
+		 */
+		void InitializeCommandInfo();
+	};
+}
diff --git a/Engine/Source/Developer/WidgetRegistration/Public/Persistence/BuilderPersistenceManager.h b/Engine/Source/Developer/WidgetRegistration/Public/Persistence/BuilderPersistenceManager.h
new file mode 100644
--- /dev/null
+++ b/Engine/Source/Developer/WidgetRegistration/Public/Persistence/BuilderPersistenceManager.h
@@ -0,0 +1,93 @@
+// Copyright Epic Games, Inc. All Rights Reserved.
+
+#pragma once
+
+#include "Containers/Map.h"
+#include "EditorConfigBase.h"
+#include "BuilderKey.h"
+#include "UObject/ObjectPtr.h"
+
+#include "BuilderPersistenceManager.generated.h"
+
+/**
+ * Manages an array of FNames to persist
+ */
+USTRUCT()
+struct WIDGETREGISTRATION_API FPersistedNameArray
+{
+	GENERATED_BODY()
+public:
+	
+	UPROPERTY()
+	TArray<FName> ArrayOfNamesToPersist;
+};
+
+
+/**
+ * The Builder Persistence Manager handles persistence for Builders through use of FBuilderKeys
+ */
+UCLASS(EditorConfig="BuilderPersistenceManager")
+class WIDGETREGISTRATION_API  UBuilderPersistenceManager : public UEditorConfigBase
+{
+	GENERATED_BODY()
+
+public:
+	/**
+	 * Returns the default named favorites array for the Builder with FBuilderKey Key, if one exists, else an empty array is returned
+	 * 
+	 * @param Key the FBuilderKey to retrieve the favorites for
+	 */
+	TArray<FName> GetPersistedFavoritesNamesArray( const UE::DisplayBuilders::FBuilderKey& Key  );
+
+	/**
+	 * Sets the default named favorites array for the Builder with FBuilderKey Key
+	 * 
+	 * @param Key the FBuilderKey to retrieve the favorites for
+	 * @param Favorites the array of FNames for the favorites to persist
+	 */
+	void SetPersistedFavoritesNamesArray( const UE::DisplayBuilders::FBuilderKey& Key, TArray<FName>& Favorites );
+
+	/**
+	 * Returns the an array for the Builder with FBuilderKey Key and the suffix that the array was persisted with, if one exists, else an empty array is returned
+	 * 
+	 * @param Key the FBuilderKey to retrieve the FNames array for
+	 * @param PersistenceKeySuffix the suffix to add to the FBuilderKey to persist the array of FNames with
+	 */
+	TArray<FName> GetPersistedArrayOfNames( const UE::DisplayBuilders::FBuilderKey& Key, FName PersistenceKeySuffix );
+
+	/**
+	 * Persists the array of FNames the Builder with FBuilderKey Key and suffix PersistenceKeySuffix
+	 * 
+	 * @param Key the FBuilderKey to retrieve the favorites for
+	 * @param PersistenceKeySuffix the suffix added to the FBuilderKey to persist the array of FNames with
+	 */
+	void PersistArrayOfNames( const UE::DisplayBuilders::FBuilderKey& Key, FName PersistenceKeySuffix, TArray<FName>& ArrayOfNamesToPersist );
+
+	/**
+	 * Initialize the Persistence manager
+	 */
+	static void Initialize();
+	
+	/**
+	* Shuts down the Persistence manager
+	*/
+	static void ShutDown();
+
+	/**
+	 * Gets the singleton for the Builder Persistence Manager
+	 */
+	static UBuilderPersistenceManager* Get()
+	{
+		return Instance;
+	}
+
+
+private:
+
+	UPROPERTY(meta=(EditorConfig))
+	TMap<FString, FPersistedNameArray> PersistenceSettings;
+
+	static FName FavoritesSuffix;
+	
+	static TObjectPtr<UBuilderPersistenceManager> Instance;
+};
\ No newline at end of file
diff --git a/Engine/Source/Developer/WidgetRegistration/Public/Styles/SlateBrushTemplates.h b/Engine/Source/Developer/WidgetRegistration/Public/Styles/SlateBrushTemplates.h
--- a/Engine/Source/Developer/WidgetRegistration/Public/Styles/SlateBrushTemplates.h
+++ b/Engine/Source/Developer/WidgetRegistration/Public/Styles/SlateBrushTemplates.h
@@ -3,9 +3,32 @@
 #pragma once
 #include "Styling/AppStyle.h"
 #include "Styling/SlateBrush.h"
+#include "Styling/StyleColors.h"
 
 /** A const FSlateBrush* Factory */
 struct WIDGETREGISTRATION_API FSlateBrushTemplates
 {
+	// Images
+	static const FSlateBrush* DragHandle();
 	static const FSlateBrush* ThinLineHorizontal();
+
+	// Colors
+	static const FSlateBrush* Transparent();
+	static const FSlateBrush* Panel();
+	static const FSlateBrush* Recessed();
+
+	/**
+	 * gets a const FSlateBrush* with the color Color
+	 *
+	 *  @param the EStyleColor we need the slate brush for
+	 */
+	const FSlateBrush* GetBrushWithColor(EStyleColor Color);
+
+	/**
+	 * gets the FSlateBrushTemplates singleton
+	 */
+	static FSlateBrushTemplates& Get();
+
+	/** the map of EStyleColor to const FSlateBrush */
+	TMap<EStyleColor, const FSlateBrush> EStyleColorToSlateBrushMap;
 };
diff --git a/Engine/Source/Developer/WidgetRegistration/Public/ToolElementRegistry.h b/Engine/Source/Developer/WidgetRegistration/Public/ToolElementRegistry.h
--- a/Engine/Source/Developer/WidgetRegistration/Public/ToolElementRegistry.h
+++ b/Engine/Source/Developer/WidgetRegistration/Public/ToolElementRegistry.h
@@ -2,6 +2,7 @@
 
 #pragma once
 
+#include "BuilderKey.h"
 #include "Widgets/SWidget.h"
 
 class FToolElementRegistry;
@@ -22,6 +23,7 @@ enum class EToolElement : uint32
 class WIDGETREGISTRATION_API FToolElementRegistrationArgs :
 	public TSharedFromThis<FToolElementRegistrationArgs>
 {
+	
 public:
 
 	/** The constructor, which takes an EToolElement which will define the type of Tool Element*/
@@ -29,6 +31,13 @@ public:
 
 	/** The constructor, which takes an EToolElement which will define the type of Tool Element*/
 	FToolElementRegistrationArgs(FName InStyleClassName);
+	
+	/**
+	 * The constructor, which takes an FBuilderKey which will define the type of Tool Element
+	 *
+	 * @param InBuilderKey the FBuilderKey that provides the Key for this builder
+	 */
+	FToolElementRegistrationArgs( UE::DisplayBuilders::FBuilderKey InBuilderKey );
 
 	/** default destructor in case any subclasses need to provide a destructor*/
 	virtual ~FToolElementRegistrationArgs() = default;
@@ -37,7 +46,7 @@ public:
 	virtual TSharedPtr<SWidget> GenerateWidget();
 
 	/** Implements the generation of the TSharedRef<SWidget> */
-	virtual TSharedRef<SWidget> GenerateWidgetSharedRef();
+	TSharedRef<SWidget> GenerateWidgetSharedRef();
 
 	/** Updates/reloads this widget. This should be called after a consumer has changed any Data in this */
 	virtual void UpdateWidget();
@@ -50,6 +59,12 @@ public:
 
 	/** The style class name */
 	const FName StyleClassName;
+	
+	/** 
+	* the FBuilderKey for this. This provides a key into persistence, if any exists for the builder, along 
+	* with other things.
+    */
+	const UE::DisplayBuilders::FBuilderKey BuilderKey;
 };
 
 /** Serves as a key into the FToolElementRegistry of FToolElements */
diff --git a/Engine/Source/Developer/WidgetRegistration/WidgetRegistration.Build.cs b/Engine/Source/Developer/WidgetRegistration/WidgetRegistration.Build.cs
--- a/Engine/Source/Developer/WidgetRegistration/WidgetRegistration.Build.cs
+++ b/Engine/Source/Developer/WidgetRegistration/WidgetRegistration.Build.cs
@@ -13,7 +13,8 @@ public class WidgetRegistration : ModuleRules
 				"Slate",
 				"Engine",
 				"ToolMenus",
-				"CoreUObject"
+				"CoreUObject",
+				"EditorConfig"
 			});
 
 		PublicIncludePathModuleNames.AddRange(
