[Backout] - CL32887250
Original CL Desc
-----------------------------------------------------------------
Place actors milestone 2, managers and keys

Added keys for Builders, Icons, Icon sizes, and managers for style, command creation, Inputs, and Persistence.

diff --git a/Engine/Source/Developer/WidgetRegistration/Private/BuilderKey.cpp b/Engine/Source/Developer/WidgetRegistration/Private/BuilderKey.cpp
deleted file mode 100644
diff --git a/Engine/Source/Developer/WidgetRegistration/Private/Common/BuilderTypes.cpp b/Engine/Source/Developer/WidgetRegistration/Private/Common/BuilderTypes.cpp
deleted file mode 100644
diff --git a/Engine/Source/Developer/WidgetRegistration/Private/Inputs/BuilderCommandCreationManager.cpp b/Engine/Source/Developer/WidgetRegistration/Private/Inputs/BuilderCommandCreationManager.cpp
deleted file mode 100644
diff --git a/Engine/Source/Developer/WidgetRegistration/Private/Inputs/BuilderCommandCreationManager.h b/Engine/Source/Developer/WidgetRegistration/Private/Inputs/BuilderCommandCreationManager.h
deleted file mode 100644
diff --git a/Engine/Source/Developer/WidgetRegistration/Private/Inputs/BuilderInputManager.cpp b/Engine/Source/Developer/WidgetRegistration/Private/Inputs/BuilderInputManager.cpp
deleted file mode 100644
diff --git a/Engine/Source/Developer/WidgetRegistration/Private/Inputs/BuilderInputManager.h b/Engine/Source/Developer/WidgetRegistration/Private/Inputs/BuilderInputManager.h
deleted file mode 100644
diff --git a/Engine/Source/Developer/WidgetRegistration/Private/Persistence/BuilderPersistenceManager.cpp b/Engine/Source/Developer/WidgetRegistration/Private/Persistence/BuilderPersistenceManager.cpp
deleted file mode 100644
diff --git a/Engine/Source/Developer/WidgetRegistration/Private/Styles/BuilderIconKeys.cpp b/Engine/Source/Developer/WidgetRegistration/Private/Styles/BuilderIconKeys.cpp
deleted file mode 100644
diff --git a/Engine/Source/Developer/WidgetRegistration/Private/Styles/BuilderIconKeys.h b/Engine/Source/Developer/WidgetRegistration/Private/Styles/BuilderIconKeys.h
deleted file mode 100644
diff --git a/Engine/Source/Developer/WidgetRegistration/Private/Styles/BuilderIconSizeKeys.h b/Engine/Source/Developer/WidgetRegistration/Private/Styles/BuilderIconSizeKeys.h
deleted file mode 100644
diff --git a/Engine/Source/Developer/WidgetRegistration/Private/Styles/BuilderStyleManager.cpp b/Engine/Source/Developer/WidgetRegistration/Private/Styles/BuilderStyleManager.cpp
deleted file mode 100644
diff --git a/Engine/Source/Developer/WidgetRegistration/Private/Styles/BuilderStyleManager.h b/Engine/Source/Developer/WidgetRegistration/Private/Styles/BuilderStyleManager.h
deleted file mode 100644
diff --git a/Engine/Source/Developer/WidgetRegistration/Private/Styles/SlateBrushTemplates.cpp b/Engine/Source/Developer/WidgetRegistration/Private/Styles/SlateBrushTemplates.cpp
--- a/Engine/Source/Developer/WidgetRegistration/Private/Styles/SlateBrushTemplates.cpp
+++ b/Engine/Source/Developer/WidgetRegistration/Private/Styles/SlateBrushTemplates.cpp
@@ -1,56 +1,8 @@
 // Copyright Epic Games, Inc. All Rights Reserved.
 
 #include "Styles/SlateBrushTemplates.h"
-#include "Brushes/SlateColorBrush.h"
-#include "Styling/CoreStyle.h"
-#include "Styling/StyleColors.h"
-
-const FSlateBrush* FSlateBrushTemplates::DragHandle()
-{
-	const FSlateBrush* Brush = FCoreStyle::Get().GetBrush("VerticalBoxDragIndicatorShort");
-	return Brush;
-}
 
 const FSlateBrush* FSlateBrushTemplates::ThinLineHorizontal()
 {
 	return FAppStyle::GetBrush( "ThinLine.Horizontal");		
 }
-
-const FSlateBrush* FSlateBrushTemplates::Transparent()
-{
-	static const FSlateColorBrush ColorBrush{ FLinearColor::Transparent };
-	return &ColorBrush;
-}
-
-const FSlateBrush* FSlateBrushTemplates::Panel()
-{
-	static const FSlateColorBrush ColorBrush{ FStyleColors::Panel };
-	return &ColorBrush;
-}
-
-const FSlateBrush* FSlateBrushTemplates::Recessed()
-{
-	static const FSlateColorBrush ColorBrush{ FStyleColors::Recessed };
-	return &ColorBrush;
-}
-
-const FSlateBrush* FSlateBrushTemplates::GetBrushWithColor(EStyleColor Color)
-{
-	const FSlateBrush* ColorBrush = EStyleColorToSlateBrushMap.Find( Color );
-	
-	if ( ColorBrush )
-	{
-		return ColorBrush;
-	}
-
-	const FSlateColorBrush NewColorBrush{ Color };
-	EStyleColorToSlateBrushMap.Add( Color, static_cast<const FSlateBrush>( NewColorBrush ) );
-
-	return EStyleColorToSlateBrushMap.Find( Color );;
-}
-
-FSlateBrushTemplates& FSlateBrushTemplates::Get()
-{
-	static FSlateBrushTemplates Templates;
-	return Templates;
-}
diff --git a/Engine/Source/Developer/WidgetRegistration/Private/ToolElementRegistry.cpp b/Engine/Source/Developer/WidgetRegistration/Private/ToolElementRegistry.cpp
--- a/Engine/Source/Developer/WidgetRegistration/Private/ToolElementRegistry.cpp
+++ b/Engine/Source/Developer/WidgetRegistration/Private/ToolElementRegistry.cpp
@@ -26,14 +26,7 @@ FToolElementRegistrationArgs::FToolElementRegistrationArgs( FName InStyleClassNa
 	StyleClassName(InStyleClassName)
 {
 }
-
-FToolElementRegistrationArgs::FToolElementRegistrationArgs( UE::DisplayBuilders::FBuilderKey InBuilderKey ) :
-	ToolElementType(EToolElement::Section)
-	, StyleClassName( InBuilderKey.ToName() )
-	, BuilderKey( InBuilderKey )
-{
-}
-
+	
 TSharedPtr<SWidget> FToolElementRegistrationArgs::GenerateWidget()
 {
 	return nullptr;
diff --git a/Engine/Source/Developer/WidgetRegistration/Private/WidgetRegistrationModule.cpp b/Engine/Source/Developer/WidgetRegistration/Private/WidgetRegistrationModule.cpp
--- a/Engine/Source/Developer/WidgetRegistration/Private/WidgetRegistrationModule.cpp
+++ b/Engine/Source/Developer/WidgetRegistration/Private/WidgetRegistrationModule.cpp
@@ -2,22 +2,16 @@
 
 #include "WidgetRegistrationModule.h"
 #include "ToolkitStyle.h"
-#include "Inputs/BuilderInputManager.h"
-#include "Persistence/BuilderPersistenceManager.h"
 
 
 void FWidgetRegistrationModule::StartupModule()
 {
 	FToolkitStyle::Initialize();
-	UBuilderPersistenceManager::Initialize();
-	FBuilderInputManager::Initialize();
 }
 	
 void FWidgetRegistrationModule::ShutdownModule()
 {
 	FToolkitStyle::Shutdown();
-	UBuilderPersistenceManager::ShutDown();
-	FBuilderInputManager::Shutdown();
 }
 
 
diff --git a/Engine/Source/Developer/WidgetRegistration/Public/BuilderKey.h b/Engine/Source/Developer/WidgetRegistration/Public/BuilderKey.h
deleted file mode 100644
diff --git a/Engine/Source/Developer/WidgetRegistration/Public/Common/BuilderTypes.h b/Engine/Source/Developer/WidgetRegistration/Public/Common/BuilderTypes.h
deleted file mode 100644
diff --git a/Engine/Source/Developer/WidgetRegistration/Public/Inputs/BuilderInput.cpp b/Engine/Source/Developer/WidgetRegistration/Public/Inputs/BuilderInput.cpp
deleted file mode 100644
diff --git a/Engine/Source/Developer/WidgetRegistration/Public/Inputs/BuilderInput.h b/Engine/Source/Developer/WidgetRegistration/Public/Inputs/BuilderInput.h
deleted file mode 100644
diff --git a/Engine/Source/Developer/WidgetRegistration/Public/Persistence/BuilderPersistenceManager.h b/Engine/Source/Developer/WidgetRegistration/Public/Persistence/BuilderPersistenceManager.h
deleted file mode 100644
diff --git a/Engine/Source/Developer/WidgetRegistration/Public/Styles/SlateBrushTemplates.h b/Engine/Source/Developer/WidgetRegistration/Public/Styles/SlateBrushTemplates.h
--- a/Engine/Source/Developer/WidgetRegistration/Public/Styles/SlateBrushTemplates.h
+++ b/Engine/Source/Developer/WidgetRegistration/Public/Styles/SlateBrushTemplates.h
@@ -3,32 +3,9 @@
 #pragma once
 #include "Styling/AppStyle.h"
 #include "Styling/SlateBrush.h"
-#include "Styling/StyleColors.h"
 
 /** A const FSlateBrush* Factory */
 struct WIDGETREGISTRATION_API FSlateBrushTemplates
 {
-	// Images
-	static const FSlateBrush* DragHandle();
 	static const FSlateBrush* ThinLineHorizontal();
-
-	// Colors
-	static const FSlateBrush* Transparent();
-	static const FSlateBrush* Panel();
-	static const FSlateBrush* Recessed();
-
-	/**
-	 * gets a const FSlateBrush* with the color Color
-	 *
-	 *  @param the EStyleColor we need the slate brush for
-	 */
-	const FSlateBrush* GetBrushWithColor(EStyleColor Color);
-
-	/**
-	 * gets the FSlateBrushTemplates singleton
-	 */
-	static FSlateBrushTemplates& Get();
-
-	/** the map of EStyleColor to const FSlateBrush */
-	TMap<EStyleColor, const FSlateBrush> EStyleColorToSlateBrushMap;
 };
diff --git a/Engine/Source/Developer/WidgetRegistration/Public/ToolElementRegistry.h b/Engine/Source/Developer/WidgetRegistration/Public/ToolElementRegistry.h
--- a/Engine/Source/Developer/WidgetRegistration/Public/ToolElementRegistry.h
+++ b/Engine/Source/Developer/WidgetRegistration/Public/ToolElementRegistry.h
@@ -2,7 +2,6 @@
 
 #pragma once
 
-#include "BuilderKey.h"
 #include "Widgets/SWidget.h"
 
 class FToolElementRegistry;
@@ -23,7 +22,6 @@ enum class EToolElement : uint32
 class WIDGETREGISTRATION_API FToolElementRegistrationArgs :
 	public TSharedFromThis<FToolElementRegistrationArgs>
 {
-	
 public:
 
 	/** The constructor, which takes an EToolElement which will define the type of Tool Element*/
@@ -31,13 +29,6 @@ public:
 
 	/** The constructor, which takes an EToolElement which will define the type of Tool Element*/
 	FToolElementRegistrationArgs(FName InStyleClassName);
-	
-	/**
-	 * The constructor, which takes an FBuilderKey which will define the type of Tool Element
-	 *
-	 * @param InBuilderKey the FBuilderKey that provides the Key for this builder
-	 */
-	FToolElementRegistrationArgs( UE::DisplayBuilders::FBuilderKey InBuilderKey );
 
 	/** default destructor in case any subclasses need to provide a destructor*/
 	virtual ~FToolElementRegistrationArgs() = default;
@@ -46,7 +37,7 @@ public:
 	virtual TSharedPtr<SWidget> GenerateWidget();
 
 	/** Implements the generation of the TSharedRef<SWidget> */
-	TSharedRef<SWidget> GenerateWidgetSharedRef();
+	virtual TSharedRef<SWidget> GenerateWidgetSharedRef();
 
 	/** Updates/reloads this widget. This should be called after a consumer has changed any Data in this */
 	virtual void UpdateWidget();
@@ -59,12 +50,6 @@ public:
 
 	/** The style class name */
 	const FName StyleClassName;
-	
-	/** 
-	* the FBuilderKey for this. This provides a key into persistence, if any exists for the builder, along 
-	* with other things.
-    */
-	const UE::DisplayBuilders::FBuilderKey BuilderKey;
 };
 
 /** Serves as a key into the FToolElementRegistry of FToolElements */
diff --git a/Engine/Source/Developer/WidgetRegistration/WidgetRegistration.Build.cs b/Engine/Source/Developer/WidgetRegistration/WidgetRegistration.Build.cs
--- a/Engine/Source/Developer/WidgetRegistration/WidgetRegistration.Build.cs
+++ b/Engine/Source/Developer/WidgetRegistration/WidgetRegistration.Build.cs
@@ -13,8 +13,7 @@ public class WidgetRegistration : ModuleRules
 				"Slate",
 				"Engine",
 				"ToolMenus",
-				"CoreUObject",
-				"EditorConfig"
+				"CoreUObject"
 			});
 
 		PublicIncludePathModuleNames.AddRange(
