Place Actors Milestone 2, new builders.

-Header and body container
- Simple Title container
- Column wrapping container
- Zero State builder
- Addition of favorites
- moved additions of placement entries into the category builder

diff --git a/Engine/Source/Developer/WidgetRegistration/Private/DataVisualization/ZeroStateBuilder.cpp b/Engine/Source/Developer/WidgetRegistration/Private/DataVisualization/ZeroStateBuilder.cpp
new file mode 100644
--- /dev/null
+++ b/Engine/Source/Developer/WidgetRegistration/Private/DataVisualization/ZeroStateBuilder.cpp
@@ -0,0 +1,55 @@
+// Copyright Epic Games, Inc. All Rights Reserved.
+
+#include "DataVisualization/ZeroStateBuilder.h"
+#include "Styling/AppStyle.h"
+#include "Widgets/Images/SImage.h"
+#include "Widgets/Text/STextBlock.h"
+#include "Widgets/Layout/SBox.h"
+#include "Widgets/SBoxPanel.h"
+
+FZeroStateBuilder::FZeroStateBuilder( UE::DisplayBuilders::FLabelAndIconArgs LabelAndIconArgs ) :
+Icon( LabelAndIconArgs.Icon )
+, Label( LabelAndIconArgs.Label )
+{
+}
+
+TSharedPtr<SWidget> FZeroStateBuilder::GenerateWidget()
+{
+	// TODO: extract style info into class
+	return SNew( SVerticalBox )
+			+ SVerticalBox::Slot()
+			.HAlign( HAlign_Center )
+			.AutoHeight()
+			[
+				SNew(SBox)
+				.MaxDesiredWidth(40.f)
+				.Padding(0.f, 20.f )
+				[
+					SNew(SImage )
+					.Image( Icon.GetIcon() )
+					.DesiredSizeOverride( FVector2D(40.f, 40.f ) )
+				]
+			]
+			+ SVerticalBox::Slot()
+			.AutoHeight()
+			[
+				SNew(SBox)
+				.MaxDesiredWidth( 250.f )
+				[
+					SNew(STextBlock)
+					.Text( Label )
+					.TextStyle( FAppStyle::Get(), "HintText")
+					.OverflowPolicy( ETextOverflowPolicy::MultilineEllipsis )
+					.Justification( ETextJustify::Center )
+					.AutoWrapText( true )
+				]
+			];
+}
+
+void FZeroStateBuilder::UpdateWidget()
+{
+}
+
+void FZeroStateBuilder::ResetWidget()
+{
+}
diff --git a/Engine/Source/Developer/WidgetRegistration/Private/DataVisualization/ZeroStateBuilder.h b/Engine/Source/Developer/WidgetRegistration/Private/DataVisualization/ZeroStateBuilder.h
new file mode 100644
--- /dev/null
+++ b/Engine/Source/Developer/WidgetRegistration/Private/DataVisualization/ZeroStateBuilder.h
@@ -0,0 +1,47 @@
+// Copyright Epic Games, Inc. All Rights Reserved.
+
+#pragma once
+
+#include "Common/BuilderTypes.h"
+#include "Layout/Containers/SlateBuilder.h"
+#include "Styles/BuilderIconKeys.h"
+
+#define LOCTEXT_NAMESPACE "ZeroStateBuilder"
+
+/**
+ * A builder which builds a view of a "Zero State" which can be displayed when no objects are available for that view. It has an icon which
+ * shows some indication of the missing data, and some text which explains the state.
+ */
+class FZeroStateBuilder : public FSlateBuilder
+{
+public:
+	/**
+	 * The constructor which takes a FLabelAndIconArgs to initialize the Zero state
+	 * 
+	 * @param LabelAndIconArgs a parameter object to provide the label and icon information 
+	 */
+	explicit FZeroStateBuilder( UE::DisplayBuilders::FLabelAndIconArgs LabelAndIconArgs =
+	UE::DisplayBuilders::FLabelAndIconArgs{ LOCTEXT("ZeroStateBuilderDefaultLabel", "No items available." )
+	, FBuilderIconKeys::Get().ZeroStateDefaultMedium().GetSlateIcon() } );
+
+	/**
+	 * Build the Zero state SWidget and return it
+	 */
+	virtual TSharedPtr<SWidget> GenerateWidget() override;
+
+private:
+	/**
+	 * Keeping UpdateWidget private until/unless there is a need for it
+	 */
+	virtual void UpdateWidget() override;
+
+	/**
+	 * Keeping ResetWidget private until/unless there is a need for it
+	 */
+	virtual void ResetWidget() override;
+
+	FSlateIcon Icon;
+	FText Label;
+};
+
+#undef LOCTEXT_NAMESPACE
\ No newline at end of file
diff --git a/Engine/Source/Developer/WidgetRegistration/Private/DataVisualization/ZeroStateBuilderTemplates.cpp b/Engine/Source/Developer/WidgetRegistration/Private/DataVisualization/ZeroStateBuilderTemplates.cpp
new file mode 100644
--- /dev/null
+++ b/Engine/Source/Developer/WidgetRegistration/Private/DataVisualization/ZeroStateBuilderTemplates.cpp
@@ -0,0 +1,46 @@
+// Copyright Epic Games, Inc. All Rights Reserved.
+
+
+#include "DataVisualization/ZeroStateBuilderTemplates.h"
+
+#define LOCTEXT_NAMESPACE "ZeroStateBuilderTemplates"
+
+const FZeroStateBuilderTemplates& FZeroStateBuilderTemplates::Get()
+{
+	static const FZeroStateBuilderTemplates Templates;
+	return Templates;
+}
+
+TSharedRef<FZeroStateBuilder> FZeroStateBuilderTemplates::GetDefault( FText DefaultText ) const
+{
+	const FText LabelText = DefaultText.IsEmpty() ?
+		LOCTEXT("ZeroStateDefaultLabel", "No results.") : DefaultText;
+	const FSlateIcon SlateIcon = FBuilderIconKeys::Get().ZeroStateDefaultMedium().GetSlateIcon();
+	
+	return MakeShared<FZeroStateBuilder>(
+		UE::DisplayBuilders::FLabelAndIconArgs
+		{
+			LabelText,
+			SlateIcon
+		}
+	);
+}
+
+TSharedRef<FZeroStateBuilder> FZeroStateBuilderTemplates::GetFavorites( FText NoFavoritesAvailableText ) const
+{
+	const FText LabelText = NoFavoritesAvailableText.IsEmpty() ?
+		LOCTEXT("ZeroStateFavoritesLabel", "No favorites." ) :
+		NoFavoritesAvailableText;
+
+	const FSlateIcon SlateIcon = FBuilderIconKeys::Get().ZeroStateFavoritesMedium().GetSlateIcon();
+	
+	return MakeShared<FZeroStateBuilder>(
+			UE::DisplayBuilders::FLabelAndIconArgs
+			{
+				LabelText,
+				SlateIcon
+			}
+		);
+}
+
+#undef LOCTEXT_NAMESPACE
\ No newline at end of file
diff --git a/Engine/Source/Developer/WidgetRegistration/Private/DataVisualization/ZeroStateBuilderTemplates.h b/Engine/Source/Developer/WidgetRegistration/Private/DataVisualization/ZeroStateBuilderTemplates.h
new file mode 100644
--- /dev/null
+++ b/Engine/Source/Developer/WidgetRegistration/Private/DataVisualization/ZeroStateBuilderTemplates.h
@@ -0,0 +1,29 @@
+// Copyright Epic Games, Inc. All Rights Reserved.
+
+#pragma once
+
+#include "ZeroStateBuilder.h"
+
+/**
+ * FZeroStateBuilderTemplates provides some templates for Zero State Builders.
+ */
+class FZeroStateBuilderTemplates
+{
+public:
+	static const FZeroStateBuilderTemplates& Get();
+
+	/**
+	 * A default Zero state builder which gives a general default icon and text for unavailable items
+	 *
+	 * @param DefaultText the text shown below the icon in the zero state view
+	 */
+	TSharedRef<FZeroStateBuilder> GetDefault( FText DefaultText = FText::GetEmpty() ) const;
+
+	/**
+	 * A Zero state builder which gives an icon and text for no favorites items available
+	 
+	 * @param NoFavoritesAvailableText the text shown below the icon in the zero state view
+	 */
+	TSharedRef<FZeroStateBuilder> GetFavorites( FText NoFavoritesAvailableText = FText::GetEmpty() ) const;
+};
+
diff --git a/Engine/Source/Developer/WidgetRegistration/Private/Inputs/BuilderCommandCreationManager.h b/Engine/Source/Developer/WidgetRegistration/Private/Inputs/BuilderCommandCreationManager.h
--- a/Engine/Source/Developer/WidgetRegistration/Private/Inputs/BuilderCommandCreationManager.h
+++ b/Engine/Source/Developer/WidgetRegistration/Private/Inputs/BuilderCommandCreationManager.h
@@ -42,7 +42,7 @@ public:
 	void UnregisterCommandForBuilder( UE::DisplayBuilders::FBuilderInput& OutBuilderInput ) const;
 
 	/**
-	 * returns the singleton of the FBuilderCommandCreationManager
+	 * @return the singleton of the FBuilderCommandCreationManager
 	 */
 	static const FBuilderCommandCreationManager& Get();
 };
\ No newline at end of file
diff --git a/Engine/Source/Developer/WidgetRegistration/Private/Inputs/BuilderInputManager.h b/Engine/Source/Developer/WidgetRegistration/Private/Inputs/BuilderInputManager.h
--- a/Engine/Source/Developer/WidgetRegistration/Private/Inputs/BuilderInputManager.h
+++ b/Engine/Source/Developer/WidgetRegistration/Private/Inputs/BuilderInputManager.h
@@ -35,12 +35,12 @@ public:
 	static void Shutdown();
 
 	/**
-	 * Returns the Command manager for Builders
+	 * @return the Command manager for Builders
 	 */
 	const FBuilderCommandCreationManager& GetCommandManager();
 	
 	/**
-	 * Returns the Drag and Drop manager for Builders
+	 * @return the Drag and Drop manager for Builders
 	 */
 	// FBuilderDragAndDropManager& GetDragAndDropManager();
 };
diff --git a/Engine/Source/Developer/WidgetRegistration/Private/Layout/CategoryDrivenContentBuilder.cpp b/Engine/Source/Developer/WidgetRegistration/Private/Layout/CategoryDrivenContentBuilder.cpp
--- a/Engine/Source/Developer/WidgetRegistration/Private/Layout/CategoryDrivenContentBuilder.cpp
+++ b/Engine/Source/Developer/WidgetRegistration/Private/Layout/CategoryDrivenContentBuilder.cpp
@@ -2,66 +2,283 @@
 
 
 #include "Layout/CategoryDrivenContentBuilder.h"
-#include "ToolkitBuilder.h"
+
+#include "Containers/ColumnWrappingContainerTemplates.h"
+#include "DataVisualization/ZeroStateBuilder.h"
+#include "DataVisualization/ZeroStateBuilderTemplates.h"
 #include "Templates/SharedPointer.h"
+#include "Framework/MultiBox/SToolBarButtonBlock.h"
+#include "Inputs/BuilderInputManager.h"
+#include "Input/DragAndDrop.h"
+// #include "Inputs/DragAndDrop/AcceptSingleItemDropZoneHandler.h"
+#include "Layout/Containers/ColumnWrappingContainer.h"
+#include "Layout/Containers/SimpleTitleContainer.h"
+#include "Persistence/BuilderPersistenceManager.h"
+#include "Layout/Containers/SlateBuilder.h"
+#include "Styling/ToolBarStyle.h"
+
+#define LOCTEXT_NAMESPACE "CategoryDrivenContentBuilder"
+
+const TArray<FName>& FCategoryDrivenContentBuilder::GetFavorites() const
+{
+	return Favorites;
+}
+
+TSharedRef<SWidget> FCategoryDrivenContentBuilder::CreateFavoritesContextMenu( FString FavoritesItemName )
+{
+	bool bInShouldCloseWindowAfterMenuSelection = true;
+	TSharedPtr< const FUICommandList > InCommandList;
+	
+	FMenuBuilder MenuBuilder(bInShouldCloseWindowAfterMenuSelection, InCommandList);
 
-FCategoryDrivenContentBuilder ::FCategoryDrivenContentBuilder(FCategoryDrivenContentBuilderArgs& Args):
-	FCategoryDrivenContentBuilderBase( Args  )
+	UBuilderPersistenceManager::Get()->PersistArrayOfNames( BuilderKey, BuilderKey.ToName(), Favorites );
+
+	const FText ItemText = Favorites.Contains( FavoritesItemName ) ?
+		LOCTEXT("CategoryDrivenContentBuilder_RemoveFromFavorites", "Remove from Favorites") :
+		LOCTEXT("CategoryDrivenContentBuilder_AddToFavorites", "Add to Favorites");
+		
+	const FUIAction ItemAction(FExecuteAction::CreateSP( this, &FCategoryDrivenContentBuilder::ToggleFavorite, FName( FavoritesItemName ) ) );
+	MenuBuilder.AddMenuEntry(ItemText, ItemText, FSlateIcon(), ItemAction);
+		
+	return MenuBuilder.MakeWidget();
+}
+
+FCategoryDrivenContentBuilder::FCategoryDrivenContentBuilder(FCategoryDrivenContentBuilderArgs& Args):
+	FCategoryDrivenContentBuilderBase(Args )
+	, FavoritesCategoryName( Args.FavoritesCommandName )
+	, CategoryLabel( Args.CategoryLabel )
+	, TitleContainer( nullptr )
+	, bIsFilledWithWidget( false )
+	, bShowNoCategorySelection( false )
+	/*, FavoritesDropZoneHandler( Args.FavoritesDropZoneHandler )*/
 {
+	Favorites = UBuilderPersistenceManager::Get()->GetPersistedFavoritesNamesArray( BuilderKey );
+	ActiveCategoryName = Args.ActiveCategoryName;
 }
 
 FCategoryDrivenContentBuilder::~FCategoryDrivenContentBuilder()
 {
-	ProvideSelectedCategoryContentDelegate.Unbind();
+	UpdateContentForCategoryDelegate.Unbind();
+}
+
+void FCategoryDrivenContentBuilder::UpdateContentForCategory( FName InActiveCategoryName, FText InActiveCategoryText )
+{
+	bShowNoCategorySelection = false;
+	ChildBuilderArray.Empty();
+
+	ActiveCategoryName = InActiveCategoryName;
+	
+	UpdateWidget();
 }
 
-void FCategoryDrivenContentBuilder::ProvideSelectedCategoryContent( FName InActiveCommandName )
+void FCategoryDrivenContentBuilder::ToggleFavorite( FName InFavoriteCommandName )
 {
-	ActiveCommandName = InActiveCommandName;
+	if  ( Favorites.Contains( InFavoriteCommandName  ) )
+	{
+		Favorites.Remove( InFavoriteCommandName );
+	}
+	else
+	{
+		Favorites.Add( InFavoriteCommandName );
+	}
+
+	UBuilderPersistenceManager::Get()->SetPersistedFavoritesNamesArray( BuilderKey, Favorites );
 	UpdateWidget();
 }
 
+void FCategoryDrivenContentBuilder::AddFavorite(FName InFavoriteCommandName)
+{
+	if ( !Favorites.Contains( InFavoriteCommandName ) )
+	{
+		ToggleFavorite( InFavoriteCommandName );
+	}
+}
+
+void FCategoryDrivenContentBuilder::AddBuilder(TSharedRef<SWidget> Widget)
+{
+	bIsFilledWithWidget = false;
+	ChildBuilderArray.Add( MakeShared<FSlateBuilder>( Widget ) );
+}
+
+void FCategoryDrivenContentBuilder::FillWithBuilder(TSharedRef<SWidget> Widget)
+{
+	ChildBuilderArray.Empty();
+	ChildBuilderArray.Add( MakeShared<FSlateBuilder>( Widget ) );
+	bIsFilledWithWidget = true;
+}
+
+void FCategoryDrivenContentBuilder::ClearCategoryContent()
+{
+	MainContentVerticalBox->ClearChildren();
+	ChildBuilderArray.Empty();
+}
 
 void FCategoryDrivenContentBuilder::InitializeCategoryToolbar()
 {
-	for ( const TSharedPtr<FUICommandInfo>& Command : ContentLoaderCommands )
+	for ( TTuple<FName, UE::DisplayBuilders::FBuilderInput>  Pair : CategoryNameToBuilderInputMap )
 	{
+		FButtonArgs ButtonArgs = Pair.Value.ButtonArgs;
+		const TSharedPtr<const FUICommandInfo> Command = ButtonArgs.Command;
+		const FName CommandName = Command->GetCommandName();
+		
 		LoadToolPaletteCommandList->MapAction(
 			Command,
-			FExecuteAction::CreateSP(this, &FCategoryDrivenContentBuilder::ProvideSelectedCategoryContent, Command->GetCommandName()),
-			FCanExecuteAction(),
-			FGetActionCheckState::CreateSP(this, &FCategoryDrivenContentBuilder::IsActiveToolPalette,
-			                               Command->GetCommandName())
+			FExecuteAction::CreateSP(this, &FCategoryDrivenContentBuilder::UpdateContentForCategory,
+				CommandName, Command->GetLabel() ),
+			FCanExecuteAction::CreateLambda([] { return true; }),
+			FGetActionCheckState::CreateLambda([this, CommandName] ()
+			{
+				if ( !bShowNoCategorySelection )
+				{
+					return ActiveCategoryName == CommandName ? ECheckBoxState::Checked : ECheckBoxState::Unchecked;		
+				}
+				return ECheckBoxState::Unchecked;
+			}
+			)
 		);
-		LoadPaletteToolBarBuilder->AddToolBarButton(Command);
+
+		LoadPaletteToolBarBuilder->AddToolBarButton( ButtonArgs );
+	}
+}
+
+TSharedPtr<SWidget> FCategoryDrivenContentBuilder::GenerateWidget()
+{
+	TSharedRef<SWidget> Widget = FCategoryDrivenContentBuilderBase::GenerateWidget().ToSharedRef();
+
+	if ( !ActiveCategoryName.IsNone() )
+	{
+		UpdateContentForCategory( ActiveCategoryName );
 	}
+
+	return Widget;
 }
 
 void FCategoryDrivenContentBuilder::UpdateWidget()
 {
-	ToolkitWidgetVBox->ClearChildren();
+	if ( !CategoryNameToBuilderInputMap.IsEmpty() )
+	{
+		CategoryLabel = FText::GetEmpty();
+	
+		if ( const UE::DisplayBuilders::FBuilderInput* Input = CategoryNameToBuilderInputMap.Find( ActiveCategoryName ) )
+		{
+			CategoryLabel = (*Input).Label;
+		
+			if ( LoadPaletteToolBarBuilder && !bShowNoCategorySelection )
+			{
+				//	LoadPaletteToolBarBuilder->SetLastSelectedCommandIndex( (*Input).Index );
+			}
+		}
+
+		FName Name = bShowNoCategorySelection ? NAME_None : ActiveCategoryName;
+		FText Text = bShowNoCategorySelection ? FText::GetEmpty() : CategoryLabel;
+		UpdateContentForCategoryDelegate.ExecuteIfBound(  Name, Text );
+		FSimpleTitleContainerArgs Args{ CategoryLabel };
+		Args.bIsHeaderHiddenOnCreate = bShowNoCategorySelection;
+		TSharedPtr<FZeroStateBuilder> ZeroStateBuilder;
+	
+		TitleContainer = MakeShared<FSimpleTitleContainer>( Args );
+		bool bIsFavoritesCategory = ActiveCategoryName == FavoritesCategoryName;
+	
+		if ( bIsFilledWithWidget && ChildBuilderArray.Num() == 1 )
+		{
+			TitleContainer->SetBody( ChildBuilderArray[0] );
+		}
+		else if ( !ChildBuilderArray.IsEmpty() )
+		{
+			if ( !ColumnWrappingContainer.IsValid() )
+			{
+				ColumnWrappingContainer = FColumnWrappingContainerTemplates::Get().GetBestFitColumnsWithSmallCells();			
+			}
 
-	if ( ProvideSelectedCategoryContentDelegate.IsBound() )
+			ColumnWrappingContainer->SetBuilders( ChildBuilderArray );
+			TitleContainer->SetBody
+			( 
+				ColumnWrappingContainer->GenerateWidgetSharedRef()
+			);
+		}
+		else
+		{
+			ZeroStateBuilder = bIsFavoritesCategory ?
+				                   FZeroStateBuilderTemplates::Get().GetFavorites(
+					                   LOCTEXT("CategoryDrivenContentBuilder_NoFavoritesYetDragActorsToAdd",
+					                           "No favorites yet.\n\n To create favorites, drag actors from other categories and drop them on the Favorites star.")) :
+
+				                   FZeroStateBuilderTemplates::Get().GetDefault(
+					                   LOCTEXT("CategoryDrivenContentBuilder_NoActorsMatchSearch",
+					                           "No actors match your search."));
+		}
+
+		const bool bIsZeroState = ZeroStateBuilder.IsValid();
+
+		MainContentVerticalBox->ClearChildren();
+		MainContentVerticalBox->AddSlot()
+		                      .FillHeight(1.0f)
+		                      .VAlign(VAlign_Fill)
+		                      .HAlign( bIsZeroState ? HAlign_Center : HAlign_Fill )
+		[
+			bIsZeroState ?
+				SNew(SBox)
+				[
+					ZeroStateBuilder->GenerateWidgetSharedRef()
+				]:
+				TitleContainer->GenerateWidgetSharedRef()
+		];
+	}
+}
+
+void FCategoryDrivenContentBuilder::SetShowNoCategorySelection(bool bInShowNoCategorySelection)
+{
+	if (bShowNoCategorySelection != bInShowNoCategorySelection)
 	{
-		TSharedRef<SWidget> Widget =  ProvideSelectedCategoryContentDelegate.Execute( ActiveCommandName );
-		ToolkitWidgetVBox->AddSlot()
-		.FillHeight(1.0f)
-		.VAlign(VAlign_Fill)
-		 [
-		   Widget
-		 ];
+		bShowNoCategorySelection = bInShowNoCategorySelection;
+		const UE::DisplayBuilders::FBuilderInput* Input = CategoryNameToBuilderInputMap.Find( ActiveCategoryName );
+
+		/*if ( bShowNoCategorySelection )
+		{
+			LoadPaletteToolBarBuilder->SetLastSelectedCommandIndex( INDEX_NONE );
+		}
+		else if ( Input != nullptr)
+		{
+			LoadPaletteToolBarBuilder->SetLastSelectedCommandIndex( (*Input).Index );
+		}*/
 	}
 }
 
 void FCategoryDrivenContentBuilder::SetCommands(TArray<TSharedPtr<FUICommandInfo>> InContentLoaderCommands)
 {
-	ContentLoaderCommands = InContentLoaderCommands;
-	const bool bForceSmallIcons = true; 
+	// TODO delete method
+}
+
+void FCategoryDrivenContentBuilder::InitializeCategoryButtons(TArray<UE::DisplayBuilders::FBuilderInput> InBuilderInputArray)
+{
+	CategoryNameToBuilderInputMap.Empty();
+	
+	for ( int32 Index = 0; Index < InBuilderInputArray.Num(); Index++ )
+	{
+		UE::DisplayBuilders::FBuilderInput& BuilderInput = InBuilderInputArray[Index];
+		BuilderInput.Index = Index;
+		
+		if ( BuilderInput.Name == FavoritesCategoryName /*&& FavoritesDropZoneHandler.IsValid()*/ )
+		{
+			/*if ( FavoritesDropZoneHandler.IsValid() )
+			{
+				 FavoritesDropZoneHandler->OnAcceptDrop =
+					FAcceptSingleItemDropZoneHandler::FOnAcceptDrop::CreateSP( this, &FCategoryDrivenContentBuilder::AddFavorite );
+				 BuilderInput.ButtonArgs.DropZoneHandler = FavoritesDropZoneHandler; 
+			}*/
+			BuilderInput.ButtonArgs.IconOverride = FSlateIcon(FAppStyle::Get().GetStyleSetName(), "Icons.Favorites.Small");
+		}
+		CategoryNameToBuilderInputMap.Add( BuilderInput.Name, BuilderInput );
+	}
+	
+	constexpr bool bForceSmallIcons = true; 
 	
-	Style = FToolkitStyle::Get().GetWidgetStyle<FToolkitWidgetStyle>("FToolkitWidgetStyle");
 	LoadToolPaletteCommandList = MakeShared<FUICommandList>();
 	LoadPaletteToolBarBuilder = MakeShared<FVerticalToolBarBuilder>(LoadToolPaletteCommandList, FMultiBoxCustomization::None, TSharedPtr<FExtender>(), bForceSmallIcons);
+	LoadPaletteToolBarBuilder->SetStyle(&FAppStyle::Get(), "FCategoryDrivenContentBuilderToolbar");
 	LoadPaletteToolBarBuilder->SetLabelVisibility( CategoryButtonLabelVisibility );
-	ToolkitWidgetVBox = SNew(SVerticalBox);
+	MainContentVerticalBox = SNew(SVerticalBox);
 	InitializeCategoryToolbar();
 }
+
+#undef LOCTEXT_NAMESPACE
diff --git a/Engine/Source/Developer/WidgetRegistration/Private/Layout/CategoryDrivenContentBuilderBase.cpp b/Engine/Source/Developer/WidgetRegistration/Private/Layout/CategoryDrivenContentBuilderBase.cpp
--- a/Engine/Source/Developer/WidgetRegistration/Private/Layout/CategoryDrivenContentBuilderBase.cpp
+++ b/Engine/Source/Developer/WidgetRegistration/Private/Layout/CategoryDrivenContentBuilderBase.cpp
@@ -1,10 +1,7 @@
 // Copyright Epic Games, Inc. All Rights Reserved.
 
 #include "Layout/CategoryDrivenContentBuilderBase.h"
-
 #include "ToolbarRegistrationArgs.h"
-
-
 #include "Widgets/SBoxPanel.h"
 #include "Widgets/Layout/SSplitter.h"
 #include "Widgets/SBoxPanel.h"
@@ -21,23 +18,18 @@ namespace UE::CategoryDrivenContentBuilderBase::Private
 
 FToolElementRegistry FCategoryDrivenContentBuilderBase::ToolRegistry = FToolElementRegistry::Get();
 
-FCategoryDrivenContentBuilderBase::FCategoryDrivenContentBuilderBase
-	(
-		FName InBuilderName
-	) :
-	FToolElementRegistrationArgs( UE::CategoryDrivenContentBuilderBase::Private::StyleName ),
-        CategoryReclickBehavior( ECategoryReclickBehavior::NoEffect) ,
-	BuilderName( InBuilderName )
+FCategoryDrivenContentBuilderBase::FCategoryDrivenContentBuilderBase ( FName InBuilderName ) :
+	FToolElementRegistrationArgs( InBuilderName )
+	, CategoryReclickBehavior( ECategoryReclickBehavior::NoEffect)
+	, BuilderName( InBuilderName )
 {
 }
 
-FCategoryDrivenContentBuilderBase::FCategoryDrivenContentBuilderBase
-	(
-		FCategoryDrivenContentBuilderArgs& Args
-	) :
-	FToolElementRegistrationArgs(Args.BuilderName ),
-	CategoryReclickBehavior( Args.CategoryReclickBehavior ),
-	BuilderName(Args.BuilderName)
+FCategoryDrivenContentBuilderBase::FCategoryDrivenContentBuilderBase( FCategoryDrivenContentBuilderArgs& Args ) :
+	FToolElementRegistrationArgs( Args.Key )
+	, CategoryReclickBehavior( Args.CategoryReclickBehavior )
+	, BuilderName(Args.BuilderName)
+	, ActiveCategoryName( Args.ActiveCategoryName )
 {
 }
 
@@ -74,8 +66,12 @@ void FCategoryDrivenContentBuilderBase::InitCategoryToolbarContainerWidget()
 
 void FCategoryDrivenContentBuilderBase::RefreshCategoryToolbarWidget()
 {
-	FToolElementRegistrationKey Key = FToolElementRegistrationKey(BuilderName, EToolElement::Toolbar);
-	VerticalToolbarElement = ToolRegistry.GetToolElementSP(Key);
+	if ( !LoadPaletteToolBarBuilder.IsValid() )
+	{
+		return;
+	}
+	FToolElementRegistrationKey ElementKey = FToolElementRegistrationKey(BuilderName, EToolElement::Toolbar);
+	VerticalToolbarElement = ToolRegistry.GetToolElementSP(ElementKey);
 	const TSharedRef<FToolbarRegistrationArgs> VerticalToolbarRegistrationArgs = MakeShareable<FToolbarRegistrationArgs>(
 		new FToolbarRegistrationArgs(LoadPaletteToolBarBuilder.ToSharedRef()));
 	
@@ -93,7 +89,11 @@ void FCategoryDrivenContentBuilderBase::RefreshCategoryToolbarWidget()
 
 TSharedPtr<SWidget> FCategoryDrivenContentBuilderBase::GenerateWidget()
 {
-	if (!ToolkitWidgetContainerVBox)
+	if ( !LoadPaletteToolBarBuilder.IsValid() )
+	{
+		return SNullWidget::NullWidget;
+	}
+	if (!ToolkitWidgetContainerVBox )
 	{
 		CreateWidget();
 	}
@@ -103,10 +103,11 @@ TSharedPtr<SWidget> FCategoryDrivenContentBuilderBase::GenerateWidget()
 
 void FCategoryDrivenContentBuilderBase::CreateWidget()
 {
-	ToolkitWidgetVBox = ToolkitWidgetVBox.IsValid() ? ToolkitWidgetVBox : SNew(SVerticalBox);
-	ToolkitWidgetVBox->ClearChildren();
+	MainContentVerticalBox = MainContentVerticalBox.IsValid() ? MainContentVerticalBox : SNew(SVerticalBox);
+	MainContentVerticalBox->ClearChildren();
 	RefreshCategoryToolbarWidget();
-	ProvideSelectedCategoryContent();
+	UpdateContentForCategory( ActiveCategoryName );
+	
 
 	ToolkitWidgetContainerVBox = SNew(SVerticalBox)
 	+ SVerticalBox::Slot().AutoHeight() [ *FSeparatorTemplates::SmallHorizontalPanelNoBorder()  ]
@@ -125,7 +126,7 @@ void FCategoryDrivenContentBuilderBase::CreateWidget()
 		+ SSplitter::Slot()
 		.SizeRule(SSplitter::FractionOfParent)
 			[
-				ToolkitWidgetVBox->AsShared()
+				MainContentVerticalBox->AsShared()
 			];
 	
 		ToolkitWidgetContainerVBox->AddSlot()
@@ -136,16 +137,21 @@ void FCategoryDrivenContentBuilderBase::CreateWidget()
 		];
 }
 
-FCategoryDrivenContentBuilderArgs::FCategoryDrivenContentBuilderArgs(FName InBuilderName):
-                                                                                                      BuilderName( InBuilderName ),
-                                                                                                      bShowCategoryButtonLabels(false),
-                                                                                                      CategoryReclickBehavior(FCategoryDrivenContentBuilderBase::ECategoryReclickBehavior::NoEffect)
+FCategoryDrivenContentBuilderArgs::FCategoryDrivenContentBuilderArgs(
+	FName InBuilderName
+	, const UE::DisplayBuilders::FBuilderKey& InKey ) :
+		Key( InKey )
+		, BuilderName( InBuilderName )
+		, bShowCategoryButtonLabels(false)
+		, CategoryReclickBehavior(FCategoryDrivenContentBuilderBase::ECategoryReclickBehavior::NoEffect)
+		, FavoritesCommandName( NAME_None )
+		, CategoryLabel( FText::GetEmpty() )
+		, ActiveCategoryName( NAME_None )
 {
 }
 
 void FCategoryDrivenContentBuilderBase::SetCategoryButtonLabelVisibility(EVisibility Visibility)
 {
-
 	CategoryButtonLabelVisibility = Visibility;
 	InitializeCategoryToolbar();
 }
diff --git a/Engine/Source/Developer/WidgetRegistration/Private/Layout/Containers/ColumnWrappingContainer.cpp b/Engine/Source/Developer/WidgetRegistration/Private/Layout/Containers/ColumnWrappingContainer.cpp
new file mode 100644
--- /dev/null
+++ b/Engine/Source/Developer/WidgetRegistration/Private/Layout/Containers/ColumnWrappingContainer.cpp
@@ -0,0 +1,77 @@
+// Copyright Epic Games, Inc. All Rights Reserved.
+
+#include "Layout/Containers/ColumnWrappingContainer.h"
+
+#include "Widgets/Layout/SBorder.h"
+#include "Styles/SlateBrushTemplates.h"
+#include "Widgets/SBoxPanel.h"
+#include "Widgets/Layout/SUniformWrapPanel.h"
+
+
+FColumnWrappingContainerArgs::FColumnWrappingContainerArgs( float InCellHeight, int32 InNumColumnsOverride, FName InIdentifier ) :
+	FWidgetContainerArgs( InIdentifier )
+	, NumColumns( InNumColumnsOverride )
+	, CellHeight( InCellHeight )
+{
+}
+
+FColumnWrappingContainer::FColumnWrappingContainer(FColumnWrappingContainerArgs&& Args):
+	FWidgetContainer ( Args )
+	, NumColumns( Args.NumColumns )
+	, CellHeight( Args.CellHeight )
+{
+	Initialize();
+}
+
+FColumnWrappingContainer::FColumnWrappingContainer(const FColumnWrappingContainerArgs& Args):
+	FWidgetContainer ( Args )
+	, NumColumns( Args.NumColumns )
+	, CellHeight( Args.CellHeight )
+{
+	Initialize();
+}
+
+FColumnWrappingContainer& FColumnWrappingContainer::SetNumColumns(const int32& InNumColumns)
+{
+	NumColumns = InNumColumns;
+	UniformWrapPanel->SetNumColumnsOverride( NumColumns );
+	return *this;
+}
+
+void FColumnWrappingContainer::Empty()
+{
+	FWidgetContainer::Empty();
+	UniformWrapPanel->ClearChildren();
+}
+
+void FColumnWrappingContainer::CreateAndPositionWidgetAtIndex(int32 ChildBuilderIndex)
+{
+	UniformWrapPanel->AddSlot()
+	[
+		GetBuilderAtIndex( ChildBuilderIndex )->GenerateWidgetSharedRef()
+	];
+}
+
+void FColumnWrappingContainer::Initialize()
+{
+	SAssignNew (MainContentSBorder,  SBorder )
+ 		.Padding(8.f)
+		.BorderImage( FSlateBrushTemplates::Get().GetBrushWithColor( EStyleColor::Panel ) )
+		[
+			SAssignNew(UniformWrapPanel, SUniformWrapPanel)
+		 .HAlign(HAlign_Fill)
+		 .SlotPadding(FMargin{4.f, 2.f, 4.f, 2.f })
+		];
+
+	if ( NumColumns > 0 )
+	{
+		UniformWrapPanel->SetNumColumnsOverride( NumColumns );
+	}
+	if ( CellHeight != TNumericLimits<float>::Min())
+	{
+		UniformWrapPanel->SetMinDesiredSlotHeight( CellHeight );
+		UniformWrapPanel->SetMaxDesiredSlotHeight( CellHeight );
+	}
+ 	
+	MainContentWidget = MainContentSBorder;
+}
diff --git a/Engine/Source/Developer/WidgetRegistration/Private/Layout/Containers/ColumnWrappingContainer.h b/Engine/Source/Developer/WidgetRegistration/Private/Layout/Containers/ColumnWrappingContainer.h
new file mode 100644
--- /dev/null
+++ b/Engine/Source/Developer/WidgetRegistration/Private/Layout/Containers/ColumnWrappingContainer.h
@@ -0,0 +1,99 @@
+// Copyright Epic Games, Inc. All Rights Reserved.
+
+#pragma once
+#include "WidgetContainer.h"
+
+class SBorder;
+class SUniformWrapPanel;
+
+/**
+* A struct that provides arguments for FColumnWrappingContainers
+*/
+class FColumnWrappingContainerArgs : public FWidgetContainerArgs
+{
+public:
+	/**
+	 * A constructor that initializes the arguments
+	 *
+	 * @param InCellHeight the height of the cells in this container
+	 * @param InNumColumnsOverride the number of columns in the container. If not provided, the number of columns is dynamic,
+	 *		  based on the container width and the width of the text of the buttons
+	 * @param InIdentifier the identifier for this container
+	 * 
+	 */
+	FColumnWrappingContainerArgs(
+	float InCellHeight = TNumericLimits<float>::Min()
+	, int32 InNumColumnsOverride = TNumericLimits<float>::Min()
+	, FName InIdentifier = "FColumnWrappingContainer" );
+
+	/** the number of columns in the container. If not provided, the number of columns is dynamic,
+	 * based on the container width and the width of the text of the buttons */	
+	int32 NumColumns;
+
+	/** the height of the cells in this container */
+	float CellHeight;
+};
+
+/**
+ * A container that will provide best fit wrapping for columns, which you can 
+ * override if needed.
+ */
+class FColumnWrappingContainer : public FWidgetContainer
+{
+public:
+	/**
+	 * Constructs the FColumnWrappingContainer
+	 *
+	 * @param Args the parameter object to initialize this container
+	 **/
+	explicit FColumnWrappingContainer( FColumnWrappingContainerArgs&& Args  );
+	
+	/**
+	 * Constructs the FColumnWrappingContainer
+	 *
+	 * @param Args the parameter object to initialize this container
+	 **/
+	explicit FColumnWrappingContainer( const FColumnWrappingContainerArgs& Args = {} );
+
+	/**
+	 * Sets the number of columns to the int32 InNumColumns, and returns a reference to this to support chaining
+	 * 
+	 * @param InNumColumns the number of columns to have in the container
+	 * @return a reference to this to support chaining
+	 */
+	FColumnWrappingContainer& SetNumColumns( const int32& InNumColumns );
+	
+	/**
+	 * Clear the container and any widget content within it
+	 */
+	virtual void Empty() override;
+
+protected:
+	/**
+     * Creates and positions within this container the Widget generated by the Builder at index "Index" in ChildBuilderArray 
+     * 
+     * @param Index the index of the builder in ChildBuilderArray to create and position the widget
+     */
+	virtual void CreateAndPositionWidgetAtIndex(int32 Index) override;
+
+private:
+	/**
+	 * Initializes this container
+	 */
+	void Initialize();
+
+	/** The SBorder that provides the look and feel for this container */
+	TSharedPtr<SBorder> MainContentSBorder;
+
+	/**
+	 * The SUniformWrapPanel that provides the layout for this container
+	 */
+	TSharedPtr<SUniformWrapPanel> UniformWrapPanel;
+
+	/** the number of columns in the container. If not provided, the number of columns is dynamic,
+	* based on the container width and the width of the text of the buttons */
+	int32 NumColumns;
+
+	/** the height of the cells in this container */
+	float CellHeight;
+};
\ No newline at end of file
diff --git a/Engine/Source/Developer/WidgetRegistration/Private/Layout/Containers/ColumnWrappingContainerTemplates.cpp b/Engine/Source/Developer/WidgetRegistration/Private/Layout/Containers/ColumnWrappingContainerTemplates.cpp
new file mode 100644
--- /dev/null
+++ b/Engine/Source/Developer/WidgetRegistration/Private/Layout/Containers/ColumnWrappingContainerTemplates.cpp
@@ -0,0 +1,17 @@
+// Copyright Epic Games, Inc. All Rights Reserved.
+
+#include "Layout/Containers/ColumnWrappingContainerTemplates.h"
+
+#include "ColumnWrappingContainer.h"
+
+TSharedRef<FColumnWrappingContainer> FColumnWrappingContainerTemplates::GetBestFitColumnsWithSmallCells()
+{
+	// TODO: extract this out into a size type
+	return MakeShared<FColumnWrappingContainer>( FColumnWrappingContainerArgs{ 32 } );			
+}
+
+FColumnWrappingContainerTemplates& FColumnWrappingContainerTemplates::Get()
+{
+	static FColumnWrappingContainerTemplates Presets;
+	return Presets;
+}
\ No newline at end of file
diff --git a/Engine/Source/Developer/WidgetRegistration/Private/Layout/Containers/ColumnWrappingContainerTemplates.h b/Engine/Source/Developer/WidgetRegistration/Private/Layout/Containers/ColumnWrappingContainerTemplates.h
new file mode 100644
--- /dev/null
+++ b/Engine/Source/Developer/WidgetRegistration/Private/Layout/Containers/ColumnWrappingContainerTemplates.h
@@ -0,0 +1,20 @@
+// Copyright Epic Games, Inc. All Rights Reserved.
+
+#pragma once
+#include "ColumnWrappingContainer.h"
+
+/**
+ * Provides templates for 
+ */
+struct FColumnWrappingContainerTemplates
+{
+	/**
+	* Gets the singleton instance of FColumnWrappingContainerTemplates
+	*/
+	static FColumnWrappingContainerTemplates& Get();
+
+	/**
+	 * Creates a TSharedRef<FColumnWrappingContainer> that will create columns on a best fit for size basis
+	 */
+	TSharedRef<FColumnWrappingContainer> GetBestFitColumnsWithSmallCells();
+};
diff --git a/Engine/Source/Developer/WidgetRegistration/Private/Layout/Containers/HeaderAndBodyContainer.cpp b/Engine/Source/Developer/WidgetRegistration/Private/Layout/Containers/HeaderAndBodyContainer.cpp
new file mode 100644
--- /dev/null
+++ b/Engine/Source/Developer/WidgetRegistration/Private/Layout/Containers/HeaderAndBodyContainer.cpp
@@ -0,0 +1,229 @@
+// Copyright Epic Games, Inc. All Rights Reserved.
+
+#include "Layout/Containers/HeaderAndBodyContainer.h"
+#include "Styles/SlateBrushTemplates.h"
+#include "Widgets/Layout/SScrollBox.h"
+#include "Widgets/Layout/SBorder.h"
+#include "Widgets/Input/SButton.h"
+#include "Widgets/Images/SImage.h"
+#include "Widgets/SBoxPanel.h"
+#include "Widgets/Layout/SBox.h"
+
+FHeaderAndBodyContainerArgs::FHeaderAndBodyContainerArgs(
+	const FName& InIdentifier,
+	const TSharedRef<FSlateBuilder>& InHeader,
+	const TSharedRef<FSlateBuilder>& InBody,
+	const bool bInIsCollapsible,
+	const bool InIsBodyHidden,
+	const bool InIsHeaderHidden ) :
+	Identifier(InIdentifier)
+	, HeaderBuilder (InHeader )
+	, BodyBuilder( InBody )
+	, bHasToggleButtonToCollapseBody( bInIsCollapsible )
+	, bIsBodyHidden( InIsBodyHidden )
+	, bIsHeaderHiddenOnCreate( InIsHeaderHidden )
+{
+}
+
+void FHeaderAndBodyContainer::SetHeader( const TSharedRef<FSlateBuilder>& InHeaderBuilder)
+{
+	HeaderBuilder = InHeaderBuilder;
+}
+
+void FHeaderAndBodyContainer::SetHeader( const TSharedRef<SWidget>& HeaderWidget)
+{
+	HeaderBuilder = MakeShared<FSlateBuilder>( HeaderWidget );;
+}
+
+void FHeaderAndBodyContainer::SetBody( const TSharedRef<FSlateBuilder>& InBodyBuilder)
+{
+	BodyBuilder = InBodyBuilder;
+}
+
+void FHeaderAndBodyContainer::SetBody( const TSharedRef<SWidget>& BodyWidget )
+{
+	BodyBuilder = MakeShared<FSlateBuilder>( BodyWidget );
+}
+
+FHeaderAndBodyContainer::FHeaderAndBodyContainer(const FHeaderAndBodyContainerArgs& Args):
+	FSlateBuilder(Args.Identifier)
+	, HeaderBuilder( Args.HeaderBuilder )
+	, BodyBuilder( Args.BodyBuilder )
+	, bHasToggleButtonToCollapseBody( Args.bHasToggleButtonToCollapseBody )
+	, bIsBodyHidden( Args.bIsBodyHidden )
+	, bIsHeaderHidden( Args.bIsHeaderHiddenOnCreate )
+{
+}
+
+TSharedPtr<SWidget> FHeaderAndBodyContainer::GenerateWidget()
+{
+
+	FCurveSequence RolloutCurve = FCurveSequence(0.0f, 1.0f, ECurveEaseFunction::CubicOut);
+
+	constexpr EStyleColor HeaderBorderColor = EStyleColor::Dropdown;
+	constexpr EStyleColor HeaderForegroundColor = EStyleColor::Foreground;
+	constexpr EStyleColor ContainerBackground = EStyleColor::Recessed;
+	const FSlateBrush* ContainerBackgroundBrush = FSlateBrushTemplates::Get().GetBrushWithColor( ContainerBackground );
+
+	// TODO: create style class to put this information in
+	const FName NoBorderButtonStyle = "NoBorder";
+	const float FillHeight1 = 1.0f;
+	const FMargin HeaderMargin{ 4.f, 0.f, 0.f, 0.f };
+	const float NoPadding = 0.f;
+
+	TSharedPtr<SButton> ExpansionButton;
+	
+	if ( !bIsBodyHidden )
+	{
+		RolloutCurve.JumpToEnd();
+	}
+	
+	TSharedRef<SVerticalBox> FullWidgetVerticalBox =
+		SNew(SVerticalBox)
+		+ SVerticalBox::Slot()
+		.AutoHeight()
+		[
+			SAssignNew( HeaderContentSBox, SBox )
+			.WidthOverride( FOptionalSize() )
+			.HeightOverride( FOptionalSize() )
+			[
+				SAssignNew( HeaderContentSBorder, SBorder )
+				.BorderImage( FSlateBrushTemplates::Get().GetBrushWithColor( HeaderBorderColor ) )
+				.Padding( NoPadding )
+				[
+					SAssignNew( ExpansionButton, SButton )
+					.Cursor( EMouseCursor::GrabHand )
+					.ButtonStyle(FCoreStyle::Get(), NoBorderButtonStyle )
+					.ButtonColorAndOpacity( FStyleColors::Transparent )
+					.ContentPadding( NoPadding )
+					.ForegroundColor( HeaderForegroundColor )
+					[
+					SNew(SHorizontalBox)
+					+ SHorizontalBox::Slot()
+						.AutoWidth()
+						.Padding( HeaderMargin )
+						.HAlign(HAlign_Left)
+						.VAlign(VAlign_Center)
+						 [
+							 bHasToggleButtonToCollapseBody ?
+							 SAssignNew( ToggleExpansionImage, SImage )
+							 .ColorAndOpacity(FSlateColor::UseForeground()) :
+							 SNullWidget::NullWidget
+						 ]
+						 + SHorizontalBox::Slot()
+							 .FillWidth( FillHeight1 )
+							 .VAlign(VAlign_Fill)
+							  [
+								  HeaderBuilder->GenerateWidgetSharedRef()
+							  ]
+					]
+				]
+			]
+		]
+		+ SVerticalBox::Slot()
+		.FillHeight(1.0f)[SNew(SBorder)
+		.BorderImage(ContainerBackgroundBrush)
+		.Padding(0.f)
+		[
+			SAssignNew(BodyContentSBox, SBox )
+			.WidthOverride( FOptionalSize() )
+			.HeightOverride( FOptionalSize() )
+			[
+				SAssignNew( BodyContentSScrollBox, SScrollBox)
+				+ SScrollBox::Slot()
+				[
+					BodyBuilder->GenerateWidgetSharedRef()
+				]
+			]
+		]
+	];
+
+	if ( bHasToggleButtonToCollapseBody )
+	{
+		ExpansionButton->SetOnClicked( FOnClicked::CreateSP(this, &FHeaderAndBodyContainer::ToggleBodyExpansionState ) );			
+	}
+
+	UpdateWidget();		
+
+	return FullWidgetVerticalBox;
+}
+
+void FHeaderAndBodyContainer::UpdateWidget()
+{
+	if ( bIsBodyHidden )
+	{
+		UpdateToBodyRemovedState();
+	}	
+	else
+	{
+		UpdateToBodyAddedState();
+	}
+	if ( bIsHeaderHidden )
+	{ 
+		UpdateToHeaderRemovedState();
+	}
+	else
+	{
+		UpdateToHeaderAddedState();
+	}
+}
+
+void FHeaderAndBodyContainer::SetHeaderHidden(bool bInIsHeaderHidden )
+{
+	bIsHeaderHidden = bInIsHeaderHidden;
+	UpdateWidget();
+}
+
+FReply FHeaderAndBodyContainer::ToggleBodyExpansionState() 
+{
+	bIsBodyHidden = !bIsBodyHidden;
+	UpdateWidget();
+	return FReply::Handled();
+}
+
+void FHeaderAndBodyContainer::UpdateToBodyRemovedState()
+{
+	if ( ToggleExpansionImage.IsValid() )
+	{
+		const FName CollapsedArrow = "TreeArrow_Collapsed";
+		ToggleExpansionImage->SetImage( FCoreStyle::Get().GetBrush( CollapsedArrow ) );
+	}
+	BodyContentSBox->SetContent( SNullWidget::NullWidget );
+	OnBodyAddedOrRemoved.ExecuteIfBound( EBodyLifeCycleEventType::Removed );		
+}
+
+void FHeaderAndBodyContainer::UpdateToBodyAddedState()
+{
+	if ( BodyContentSScrollBox.IsValid() )
+	{
+		BodyContentSBox = SNew( SBox );
+		BodyContentSBox->SetContent( BodyContentSScrollBox.ToSharedRef() );
+	
+		if ( ToggleExpansionImage.IsValid() )
+		{
+			const FName ExpandedArrow = "TreeArrow_Expanded";
+			ToggleExpansionImage->SetImage( FCoreStyle::Get().GetBrush( ExpandedArrow ) );
+		}
+		OnBodyAddedOrRemoved.ExecuteIfBound( EBodyLifeCycleEventType::Added );		
+	}
+}
+
+void FHeaderAndBodyContainer::UpdateToHeaderRemovedState()
+{
+	if ( HeaderContentSBox.IsValid() )
+	{
+		HeaderContentSBox->SetContent( SNullWidget::NullWidget );		
+	}
+}
+
+void FHeaderAndBodyContainer::UpdateToHeaderAddedState()
+{
+	if ( HeaderContentSBorder.IsValid() && HeaderContentSBox.IsValid() )
+	{
+		HeaderContentSBox->SetContent( HeaderContentSBorder.ToSharedRef() );
+	}
+}
+
+void FHeaderAndBodyContainer::ResetWidget()
+{
+}
diff --git a/Engine/Source/Developer/WidgetRegistration/Private/Layout/Containers/HeaderAndBodyContainer.h b/Engine/Source/Developer/WidgetRegistration/Private/Layout/Containers/HeaderAndBodyContainer.h
new file mode 100644
--- /dev/null
+++ b/Engine/Source/Developer/WidgetRegistration/Private/Layout/Containers/HeaderAndBodyContainer.h
@@ -0,0 +1,166 @@
+// Copyright Epic Games, Inc. All Rights Reserved.
+
+#pragma once
+
+#include "Widgets/Layout/SExpandableArea.h"
+#include "Layout/Containers/SlateBuilder.h"
+
+class SScrollBox;
+class SBox;
+class SImage;
+
+/**
+ * The parameter objects for a FHeaderAndBodyContainer.
+ */
+class FHeaderAndBodyContainerArgs
+{
+public:
+	FHeaderAndBodyContainerArgs(
+		const FName& InIdentifier = "FHeaderAndBodyContainer",
+		const TSharedRef<FSlateBuilder>& InHeader = MakeShared<FSlateBuilder>(),
+		const TSharedRef<FSlateBuilder>& InBody = MakeShared<FSlateBuilder>(),
+		const bool bInHasToggleButtonToCollapseBody = false,
+		const bool bIsBodyHidden = false,
+		const bool bIsHeaderHiddenOnCreate = false );
+
+	/** the identifier for the FHeaderAndBodyContainer */
+	FName Identifier;
+
+	/** the FSlateBuilder that builds the header of this container */
+	TSharedRef<FSlateBuilder> HeaderBuilder;
+	
+	/** the FSlateBuilder that builds the body of this container */
+	TSharedRef<FSlateBuilder> BodyBuilder;
+
+	/** whether or not this container has a button to toggle the body open and closed */
+	bool bHasToggleButtonToCollapseBody;
+
+	/** if true, the body is initially hidden */
+	bool bIsBodyHidden;
+
+	/** if true, the header is initially hidden */
+	bool bIsHeaderHiddenOnCreate;
+};
+
+/**
+ * A container which provides a header and a body, both of which can be any SWidget.
+ */
+class  FHeaderAndBodyContainer : public FSlateBuilder
+{
+
+public:
+
+	/**
+	 * an enum to tell whether the body has been added or removed
+	 */
+	enum class EBodyLifeCycleEventType
+	{
+		Added,
+		Removed
+	};
+	
+	DECLARE_DELEGATE_OneParam( FOnBodyAddedOrRemoved,  EBodyLifeCycleEventType LifeCycleEventType );
+
+	/**
+	 * Sets the FSlateBuilder InHeaderBuilder to build the widget for the header of this container
+	 */
+	void SetHeader( const TSharedRef<FSlateBuilder>& InHeaderBuilder );
+
+	/**
+	 * Converts the SWidget HeaderWidget to an FSlateBuilder which will build the widget for the header of this container
+	 */
+	void SetHeader( const TSharedRef<SWidget>& HeaderWidget );
+
+	/**
+	 * Sets the FSlateBuilder InBodyBuilder to build the widget for the body of this container
+	 */
+	void SetBody( const TSharedRef<FSlateBuilder>& InBodyBuilder );
+	
+	/**
+	 * Converts the SWidget BodyWidget to an FSlateBuilder which will build the widget for the body of this container
+	 */
+	void SetBody( const TSharedRef<SWidget>& BodyWidget );
+
+	/**
+	 * Constructs this container.
+	 *
+	 * @param Args the FHeaderAndBodyContainerArgs which will initialize this FHeaderAndBodyContainer
+	 */
+	FHeaderAndBodyContainer( const FHeaderAndBodyContainerArgs& Args = FHeaderAndBodyContainerArgs() );
+	
+	FOnBodyAddedOrRemoved OnBodyAddedOrRemoved;
+
+	/**
+	 * Generates the widget for this header and body container and returns it
+	 */
+	virtual TSharedPtr<SWidget> GenerateWidget()  override;
+
+	/** Updates/reloads this widget. This should be called after a consumer has changed any Data in this */
+	virtual void UpdateWidget() override;
+
+	/**
+	 * Sets the header content to be hidden
+	 *
+	 * @param bInIsHeaderHidden if true the header is hidden, else it is false
+	 */
+	void SetHeaderHidden(bool bInIsHeaderHidden );
+
+	/** @return true if the body is empty, else returns false */
+	bool IsBodyEmpty()
+	{
+		return BodyBuilder->IsEmpty();
+	}
+	
+private:
+	/**
+	 * called when the user toggles the expansion of the body
+	 */
+	FReply ToggleBodyExpansionState();
+
+	/*** Makes the body visible by adding it */
+	void UpdateToBodyAddedState();
+
+	/*** Makes the body invisible by removing it */
+	void UpdateToBodyRemovedState();
+	
+	/*** Makes the header invisible by removing it */
+	void UpdateToHeaderRemovedState();
+
+	/*** Makes the header visible by adding it */
+	void UpdateToHeaderAddedState();
+
+	/** overriding to hide the method, as it is not needed */
+	virtual void ResetWidget() override;
+
+	/** builds the header for this container */
+	TSharedRef<FSlateBuilder> HeaderBuilder;
+
+	/** the SBox containing the header  */
+	TSharedPtr<SBox> HeaderContentSBox;
+	
+	/** the SBorder containing the header  */
+	TSharedPtr<SBorder> HeaderContentSBorder;
+
+	/** builds the body for this container */
+	TSharedRef<FSlateBuilder> BodyBuilder;
+	
+	/** the SBox containing the body  */
+	TSharedPtr<SBox> BodyContentSBox;
+
+	/** the scrollbox for the body */
+	TSharedPtr<SScrollBox> BodyContentSScrollBox;
+
+	/** the image for the toggle body button */
+	TSharedPtr<SImage> ToggleExpansionImage;
+
+	/** if true, this container has a toggle button to hide and show the body */
+	bool bHasToggleButtonToCollapseBody;
+
+	/** if true, the body will not be viewable */
+	bool bIsBodyHidden;
+	
+	/** if true, the header will not be viewable */
+	bool bIsHeaderHidden;
+};
+
+
diff --git a/Engine/Source/Developer/WidgetRegistration/Private/Layout/Containers/SimpleTitleContainer.cpp b/Engine/Source/Developer/WidgetRegistration/Private/Layout/Containers/SimpleTitleContainer.cpp
new file mode 100644
--- /dev/null
+++ b/Engine/Source/Developer/WidgetRegistration/Private/Layout/Containers/SimpleTitleContainer.cpp
@@ -0,0 +1,29 @@
+// Copyright Epic Games, Inc. All Rights Reserved.
+
+#include "Layout/Containers/SimpleTitleContainer.h"
+
+#include "Styling/StyleColors.h"
+#include "Widgets/SBoxPanel.h"
+#include "Widgets/Text/STextBlock.h"
+
+FSimpleTitleContainer::FSimpleTitleContainer(const FSimpleTitleContainerArgs& Args):
+	FHeaderAndBodyContainer( Args ) 
+	, Title(Args.Title)
+{
+	TSharedRef<FSlateBuilder> Header = MakeShared<FSlateBuilder>
+	(
+		SNew (SHorizontalBox )
+		+ SHorizontalBox::Slot().HAlign( HAlign_Left ).VAlign( VAlign_Center ).FillWidth( 1.f )
+		.Padding(2.f, 4.f, 0.f, 4.f )
+		[
+			SNew(STextBlock)
+			.Text( Title )
+		]
+	);
+	SetHeader( Header );
+}
+
+FSimpleTitleContainer::FSimpleTitleContainer( FSimpleTitleContainerArgs&& Args ) :
+	FSimpleTitleContainer( Args )
+{
+}
diff --git a/Engine/Source/Developer/WidgetRegistration/Private/Layout/Containers/SimpleTitleContainer.h b/Engine/Source/Developer/WidgetRegistration/Private/Layout/Containers/SimpleTitleContainer.h
new file mode 100644
--- /dev/null
+++ b/Engine/Source/Developer/WidgetRegistration/Private/Layout/Containers/SimpleTitleContainer.h
@@ -0,0 +1,45 @@
+// Copyright Epic Games, Inc. All Rights Reserved.
+
+#pragma once
+#include "HeaderAndBodyContainer.h"
+
+/**
+ * a struct to provide arguments to the FSimpleTitleContainer 
+ */
+class FSimpleTitleContainerArgs : public FHeaderAndBodyContainerArgs
+{
+public:
+	FSimpleTitleContainerArgs(
+		FText InTitle = FText::GetEmpty(),
+		const FName& InIdentifier = "FSimpleTitleContainer",
+		const TSharedRef<FSlateBuilder>& InHeader = MakeShared<FSlateBuilder>(),
+		const TSharedRef<FSlateBuilder>& InBody = MakeShared<FSlateBuilder>(),
+		const bool bIsExpanded = false ) :
+	FHeaderAndBodyContainerArgs(InIdentifier , InHeader, InBody, bIsExpanded )
+		, Title( InTitle )    
+	{
+	}
+
+	FText Title;
+};
+
+/**
+ * A container providing a simple FText title and a customizable body
+ */
+class FSimpleTitleContainer : public FHeaderAndBodyContainer
+{
+public:
+	/**
+	 * A constructor taking const FSimpleTitleContainerArgs& to initialize it
+	 */
+	FSimpleTitleContainer( const FSimpleTitleContainerArgs& Args );
+	
+	/**
+	 * A constructor taking const FSimpleTitleContainerArgs& to initialize it
+	 */
+	FSimpleTitleContainer( FSimpleTitleContainerArgs&& Args );
+
+private:
+	/** the title of the container, it will be set in the header */
+	FText Title;
+};
\ No newline at end of file
diff --git a/Engine/Source/Developer/WidgetRegistration/Private/Layout/Containers/SlateBuilder.cpp b/Engine/Source/Developer/WidgetRegistration/Private/Layout/Containers/SlateBuilder.cpp
new file mode 100644
--- /dev/null
+++ b/Engine/Source/Developer/WidgetRegistration/Private/Layout/Containers/SlateBuilder.cpp
@@ -0,0 +1,50 @@
+// Copyright Epic Games, Inc. All Rights Reserved.
+
+#include "Layout/Containers/SlateBuilder.h"
+
+FSlateBuilderArgs::FSlateBuilderArgs(const FName& InName, const TSharedPtr<SWidget> InContent) :
+	BuilderKey( UE::DisplayBuilders::FBuilderKeys::Get().None() )
+	, Content( InContent )
+{
+}
+
+FSlateBuilderArgs::FSlateBuilderArgs( const UE::DisplayBuilders::FBuilderKey& InBuilderKey, const TSharedPtr<SWidget> InContent) :
+	BuilderKey( InBuilderKey )
+	, Content( InContent )
+{
+}
+
+FSlateBuilder::FSlateBuilder( TSharedPtr<SWidget> Content, FName InIdentifier ) :
+	FToolElementRegistrationArgs( InIdentifier )
+	, SlateContent( Content )
+{
+}
+
+FSlateBuilder::FSlateBuilder( UE::DisplayBuilders::FBuilderKey InBuilderKey ) :
+	FToolElementRegistrationArgs( InBuilderKey )
+{
+}
+
+FSlateBuilder::FSlateBuilder( FName InIdentifier ) :
+	FToolElementRegistrationArgs( InIdentifier )
+{
+}
+
+TSharedPtr<SWidget> FSlateBuilder::GenerateWidget()
+{
+	if ( SlateBuilder.IsValid() )
+	{
+		SlateContent = SlateBuilder->GenerateWidget();
+	}
+	return SlateContent;
+}
+
+bool FSlateBuilder::IsEmpty()
+{
+	return !SlateContent.IsValid() || SlateContent == SNullWidget::NullWidget;
+}
+
+void FSlateBuilder::Empty()
+{
+	SlateContent = SNullWidget::NullWidget;
+}
diff --git a/Engine/Source/Developer/WidgetRegistration/Private/Layout/Containers/SlateBuilder.h b/Engine/Source/Developer/WidgetRegistration/Private/Layout/Containers/SlateBuilder.h
new file mode 100644
--- /dev/null
+++ b/Engine/Source/Developer/WidgetRegistration/Private/Layout/Containers/SlateBuilder.h
@@ -0,0 +1,93 @@
+// Copyright Epic Games, Inc. All Rights Reserved.
+
+#pragma once
+#include "ToolElementRegistry.h"
+#include "Widgets/SNullWidget.h"
+
+/**
+ * An argument object to instantiate the FSlateBuilder
+ */
+class FSlateBuilderArgs
+{
+public:
+	
+	/**
+	 * Creates the slate builder parameter object.
+	 * 
+	 * @param InName the name of the builder
+	 * @param InContent the shared pointer to the SWidget that provides the content for this slate builder. If this is nullptr or it is not
+	 * provided, this builder will show a NullWidget
+	 */
+	FSlateBuilderArgs( const FName& InName = "SlateBuilder", const TSharedPtr<SWidget> InContent = SNullWidget::NullWidget.ToSharedPtr() );
+
+	/**
+	 * Creates the slate builder parameter object.
+	 * 
+	 * @param InBuilderKey the FBuilderKey for this builder
+	 * @param InContent the shared pointer to the SWidget that provides the content for this slate builder. If this is nullptr or it is not
+	 * provided, this builder will show a NullWidget
+	 */
+	FSlateBuilderArgs( const UE::DisplayBuilders::FBuilderKey& InBuilderKey = UE::DisplayBuilders::FBuilderKeys::Get().None(), const TSharedPtr<SWidget> InContent = nullptr );
+
+	/** the BuilderKey for the arguments  */
+	const UE::DisplayBuilders::FBuilderKey BuilderKey;
+
+	/** the shared pointer to the SWidget that provides the content for this slate builder  */
+	const TSharedPtr<SWidget> Content;
+};
+
+/**
+ * A builder which can build raw Slate. This is mostly used for adding raw Slate to Builder Containers, but can be used to build any Slate. It allows
+ * raw slate to be used in the context of builder containers, and provides backward compatibility so that Slate and builders can be used in the same
+ * context.
+ */
+class FSlateBuilder : public FToolElementRegistrationArgs
+{
+public:
+
+	/**
+	 * A constructor which takes the TSharedPtr<SWidget> that provides the content for this builder, and the identifier for the builder
+	 *
+	 * @param InContent the TSharedPtr<SWidget> that provides the content for this builder
+	 * @param InIdentifier the identifier for this container
+	 */
+	explicit FSlateBuilder( TSharedPtr<SWidget> InContent = nullptr, FName InIdentifier = "FSlateBuilder" );
+
+	/**
+	 * A constructor providing the identifier for the builder
+	 *
+	 * @param InIdentifier the identifier for this container
+	 */
+	explicit FSlateBuilder( FName InIdentifier  );
+
+	/**
+	 * A constructor providing an FBuilderKey for the builder
+	 *
+	 * @param InBuilderKey the FBuilderKey for this FSlateBuilder
+	 */
+	FSlateBuilder( UE::DisplayBuilders::FBuilderKey InBuilderKey );
+
+	/**
+	 * @return the TSharedPtr<SWidget> specified as the content for this FSlateBuilder
+	 */
+	virtual TSharedPtr<SWidget> GenerateWidget() override;
+	
+	/**
+	 * @return true if the content of this is equivalent to an SNullWidget
+	 */
+	bool IsEmpty();
+	
+	/**
+	 * Empties the content
+	 */
+	virtual void Empty();
+
+protected:
+	/** The TSharedPtr<SWidget> that is the content for this FSlateBuilder */
+	TSharedPtr<SWidget> SlateContent;
+	
+private:
+	
+	/** The FSlateBuilder that provides the content for this FSlateBuilder */
+	TSharedPtr<FSlateBuilder> SlateBuilder;
+};
\ No newline at end of file
diff --git a/Engine/Source/Developer/WidgetRegistration/Private/Layout/Containers/WidgetContainer.cpp b/Engine/Source/Developer/WidgetRegistration/Private/Layout/Containers/WidgetContainer.cpp
new file mode 100644
--- /dev/null
+++ b/Engine/Source/Developer/WidgetRegistration/Private/Layout/Containers/WidgetContainer.cpp
@@ -0,0 +1,72 @@
+// Copyright Epic Games, Inc. All Rights Reserved.
+
+#include "Layout/Containers/WidgetContainer.h"
+#include "Layout/Containers/SlateBuilder.h"
+#include "Widgets/SNullWidget.h"
+
+FWidgetContainer::FWidgetContainer(const FWidgetContainerArgs Args) :
+	FSlateBuilder ( nullptr, Args.Identifier )
+{
+}
+
+FWidgetContainerArgs::FWidgetContainerArgs( FName InIdentifier) :
+	Identifier( InIdentifier )
+{
+}
+
+FWidgetContainer& FWidgetContainer::AddBuilder(const TSharedRef<FSlateBuilder>& Builder)
+{
+	ChildBuilderArray.Add( Builder );
+	return *this;
+}
+
+TSharedPtr<FSlateBuilder> FWidgetContainer::GetBuilderAtIndex(int32 Index)
+{
+	if ( Index < ChildBuilderArray.Num() )
+	{
+		return ChildBuilderArray[Index];
+	}
+	return nullptr;
+}
+
+void FWidgetContainer::Empty()
+{
+	ChildBuilderArray.Empty();
+}
+
+FWidgetContainer& FWidgetContainer::AddBuilder( TSharedRef<SWidget>  Widget)
+{
+	ChildBuilderArray.Add( MakeShared<FSlateBuilder>( Widget ) );
+	return *this;
+}
+
+TSharedPtr<SWidget> FWidgetContainer::GenerateWidget()
+{
+	for ( int32 ChildBuilderIndex = 0; ChildBuilderIndex < ChildBuilderArray.Num(); ChildBuilderIndex++)
+	{
+		CreateAndPositionWidgetAtIndex( ChildBuilderIndex );
+	}
+
+	return MainContentWidget.IsValid() ?  MainContentWidget : SNullWidget::NullWidget.ToSharedPtr();
+}
+
+void FWidgetContainer::UpdateWidget()
+{
+	for ( int32 ChildBuilderIndex = 0; ChildBuilderIndex < ChildBuilderArray.Num(); ChildBuilderIndex++)
+	{
+		ChildBuilderArray[ChildBuilderIndex]->UpdateWidget();
+	}
+}
+
+FWidgetContainer& FWidgetContainer::SetBuilders( TArray<TSharedRef<FSlateBuilder>> Builders )
+{
+	Empty();
+	
+	for ( const TSharedRef<FSlateBuilder>& Builder : Builders )
+	{
+		ChildBuilderArray.Add( Builder );
+	}
+	return *this;
+}
+
+
diff --git a/Engine/Source/Developer/WidgetRegistration/Private/Layout/Containers/WidgetContainer.h b/Engine/Source/Developer/WidgetRegistration/Private/Layout/Containers/WidgetContainer.h
new file mode 100644
--- /dev/null
+++ b/Engine/Source/Developer/WidgetRegistration/Private/Layout/Containers/WidgetContainer.h
@@ -0,0 +1,102 @@
+// Copyright Epic Games, Inc. All Rights Reserved.
+
+#pragma once
+#include "Layout/Containers/SlateBuilder.h"
+
+class SVerticalBox;
+
+/**
+ * A struct to provide the arguments for an FWidgetContainer
+ */
+class FWidgetContainerArgs
+{
+public:
+	/**
+	 * The constructor for the widget arguments
+	 * 
+	 * @param InIdentifier the identifier for this container
+	 */
+	FWidgetContainerArgs( FName InIdentifier = "FWidgetContainer" );
+
+	
+	/** the identifier for this container */
+	FName Identifier;
+};
+
+/**
+ * A Widget Container that contains an array of FSlateBuilders
+ */
+class FWidgetContainer : public  FSlateBuilder
+{
+	
+public:
+	/**
+	 * A constructor with const Args which can be used with Presets of extending containers
+	 * 
+	 * @param Args the arugments to initializes this container
+	 */
+	explicit FWidgetContainer( const FWidgetContainerArgs Args );
+
+	/**
+	 * Generates the container and its children.
+	 */
+	virtual TSharedPtr<SWidget> GenerateWidget() override;
+
+	/**
+	 * Updates the container and its children.
+	 */
+	virtual void UpdateWidget() override;
+
+	/**
+	 * Sets the FSlateBuilders in Builders as the children of this container to provide the contents of this container. If there were previously any builders in
+	 * this container, they will be cleared and these Builders will provide the entire content of this container
+	 *
+	 * @param Widgets the SWidgets that will provide the contents of this container 
+	 */
+	FWidgetContainer& SetBuilders( TArray<TSharedRef<FSlateBuilder>> Builders );
+
+	/**
+	 * Converts "Widget" to an FSlateBuilder and adds it to the existing FSlateBuilders in this container
+	 *
+	 * @param Widget the widget to add to this container
+	 */
+	FWidgetContainer& AddBuilder( TSharedRef<SWidget> Widget);
+
+	/**
+	 * Adds Builder to the contents
+	 *
+	 * @param Builder the shared pointer to the FSlateBuilder to add to this container
+	 */
+	FWidgetContainer& AddBuilder( const TSharedRef<FSlateBuilder>& Builder );
+	
+	/**
+	 * Clear the container and any widget content within it
+	 */
+	virtual void Empty() override;
+
+protected:
+	/**
+	 * The SWidget which holds the main content of this container
+	 */
+	TSharedPtr<SWidget> MainContentWidget;
+
+	/**
+	 * @return the Builder at index "Index" in ChildBuilderArray
+	 * 
+	 * @param Index the index of the builder in ChildBuilderArray to return
+	 */
+	TSharedPtr<FSlateBuilder> GetBuilderAtIndex( int32 Index );
+
+	/**
+	 * Positions the Builder at index "Index" in ChildBuilderArray within this container
+	 * 
+	 * @param Index the index of the builder in ChildBuilderArray to position
+	 */
+	virtual void CreateAndPositionWidgetAtIndex( int32 Index ) = 0;
+
+private:
+	/**
+	 * The array of FSlateBuilders that will be converted to widgets by this container
+	 */
+	TArray<TSharedRef<FSlateBuilder>> ChildBuilderArray;
+};
diff --git a/Engine/Source/Developer/WidgetRegistration/Private/Styles/BuilderIconKeys.h b/Engine/Source/Developer/WidgetRegistration/Private/Styles/BuilderIconKeys.h
--- a/Engine/Source/Developer/WidgetRegistration/Private/Styles/BuilderIconKeys.h
+++ b/Engine/Source/Developer/WidgetRegistration/Private/Styles/BuilderIconKeys.h
@@ -51,7 +51,7 @@ public:
 	
 
 	/**
-	 * returns the FSlateIcon for this FBuilderIconKey
+	 * @return the FSlateIcon for this FBuilderIconKey
 	 */
 	FSlateIcon GetSlateIcon() const;
 
diff --git a/Engine/Source/Developer/WidgetRegistration/Private/ToolkitBuilder.cpp b/Engine/Source/Developer/WidgetRegistration/Private/ToolkitBuilder.cpp
--- a/Engine/Source/Developer/WidgetRegistration/Private/ToolkitBuilder.cpp
+++ b/Engine/Source/Developer/WidgetRegistration/Private/ToolkitBuilder.cpp
@@ -464,18 +464,18 @@ EVisibility FToolkitBuilder::GetActiveToolTitleVisibility() const
 	return ActiveToolDisplayName.IsEmpty() ? EVisibility::Collapsed : EVisibility::Visible;
 }
 
-void FToolkitBuilder::ProvideSelectedCategoryContent( FName ActiveCategoryName )
+void FToolkitBuilder::UpdateContentForCategory( FName InActiveCategoryName, FText InActiveCategoryText )
 {
 	TSharedPtr<SHorizontalBox> ToolNameHeaderBox;
 	
-	if ( !ToolkitWidgetVBox.IsValid() )
+	if ( !MainContentVerticalBox.IsValid() )
 	{
 		return;
 	}
 	
 	if (ToolkitSections->ModeWarningArea)
 	{
-		ToolkitWidgetVBox->AddSlot()
+		MainContentVerticalBox->AddSlot()
 		.AutoHeight()
 		.HAlign(HAlign_Fill)
 		.Padding(5)
@@ -484,7 +484,7 @@ void FToolkitBuilder::ProvideSelectedCategoryContent( FName ActiveCategoryName )
 		];
 	}
 	
-	ToolkitWidgetVBox->AddSlot()
+	MainContentVerticalBox->AddSlot()
 	.AutoHeight()
 	.HAlign(HAlign_Fill)
 	.Padding(0)
@@ -492,12 +492,12 @@ void FToolkitBuilder::ProvideSelectedCategoryContent( FName ActiveCategoryName )
 			GetToolPaletteWidget()
 		];
 
-	ToolkitWidgetVBox->AddSlot()
+	MainContentVerticalBox->AddSlot()
 		.AutoHeight()
 		[ *FSeparatorTemplates::SmallHorizontalBackgroundNoBorder().BindVisibility(
 			TAttribute<EVisibility>::CreateLambda([this] () { return ActivePaletteButtonVisibility; }))];
 	
-	ToolkitWidgetVBox->AddSlot()
+	MainContentVerticalBox->AddSlot()
 		.AutoHeight()
 		.HAlign(HAlign_Fill)
 		.Padding(0)
@@ -541,7 +541,7 @@ void FToolkitBuilder::ProvideSelectedCategoryContent( FName ActiveCategoryName )
 
 	if (ToolkitSections->ToolWarningArea)
 	{
-		ToolkitWidgetVBox->AddSlot()
+		MainContentVerticalBox->AddSlot()
 		.AutoHeight()
 		.HAlign(HAlign_Fill)
 		.Padding(5)
@@ -552,7 +552,7 @@ void FToolkitBuilder::ProvideSelectedCategoryContent( FName ActiveCategoryName )
 
 	if (ToolkitSections->DetailsView)
 	{
-		ToolkitWidgetVBox->AddSlot()
+		MainContentVerticalBox->AddSlot()
 		.HAlign(HAlign_Fill)
 		.FillHeight(1.f)
 			[
@@ -567,7 +567,7 @@ void FToolkitBuilder::ProvideSelectedCategoryContent( FName ActiveCategoryName )
 	}
 	if (ToolkitSections->Footer)
 	{
-		ToolkitWidgetVBox->AddSlot()
+		MainContentVerticalBox->AddSlot()
 			.AutoHeight()
 			.HAlign(HAlign_Fill)
 			.VAlign(VAlign_Bottom)
diff --git a/Engine/Source/Developer/WidgetRegistration/Public/BuilderKey.h b/Engine/Source/Developer/WidgetRegistration/Public/BuilderKey.h
--- a/Engine/Source/Developer/WidgetRegistration/Public/BuilderKey.h
+++ b/Engine/Source/Developer/WidgetRegistration/Public/BuilderKey.h
@@ -25,14 +25,14 @@ namespace UE::DisplayBuilders
 		WIDGETREGISTRATION_API FName ToName() const;
 
 		/** 
-		 * Returns a combined key with Suffix
+		 * @return a combined key with Suffix
 		 * 
 		 * @param Suffix the FName that will provide the Suffix for the combined Key
 		 */
 		FString  GetKeyWithSuffix( const FName Suffix ) const;
 
 		/**
-		* returns true if this builder is set to None, meaning that it has not been initialized
+		* @return true if this builder is set to None, meaning that it has not been initialized
 		*/
 		bool IsNone() const;
 		
diff --git a/Engine/Source/Developer/WidgetRegistration/Public/Inputs/BuilderInput.h b/Engine/Source/Developer/WidgetRegistration/Public/Inputs/BuilderInput.h
--- a/Engine/Source/Developer/WidgetRegistration/Public/Inputs/BuilderInput.h
+++ b/Engine/Source/Developer/WidgetRegistration/Public/Inputs/BuilderInput.h
@@ -88,7 +88,7 @@ namespace UE::DisplayBuilders
 		static FBuilderInput NullInput;
 		
 		/**
-		 * returns true if the input name is equivalent to Name_NONE
+		 * @return true if the input name is equivalent to Name_NONE
 		 */
 		bool IsNameNone() const;
 
diff --git a/Engine/Source/Developer/WidgetRegistration/Public/Layout/CategoryDrivenContentBuilder.h b/Engine/Source/Developer/WidgetRegistration/Public/Layout/CategoryDrivenContentBuilder.h
--- a/Engine/Source/Developer/WidgetRegistration/Public/Layout/CategoryDrivenContentBuilder.h
+++ b/Engine/Source/Developer/WidgetRegistration/Public/Layout/CategoryDrivenContentBuilder.h
@@ -4,7 +4,19 @@
 #include "ToolkitBuilder.h"
 #include "ToolkitStyle.h"
 #include "Layout/CategoryDrivenContentBuilderBase.h"
+#include "Inputs/BuilderInput.h"
+#include "Misc/CoreMiscDefines.h"
 
+class FColumnWrappingContainer;
+class FSimpleTitleContainer;
+// class FAcceptSingleItemDropZoneHandler;
+class FSlateBuilder;
+
+namespace UE::DisplayBuilders
+{
+	class FBuilderInput;
+}
+	
 /**
  * A builder which creates a widget that has a vertical toolbar category picker on the left
  * hand side which populates the content on the right side.
@@ -12,13 +24,31 @@
 class  FCategoryDrivenContentBuilder : public FCategoryDrivenContentBuilderBase
 {
 public:
+	
+	
 	DECLARE_DELEGATE_RetVal_OneParam( TSharedRef<SWidget>, FProvideSelectedCategoryContent, const FName&)
+	FProvideSelectedCategoryContent  ProvideSelectedCategoryContentDelegate;
+
+	DECLARE_DELEGATE_TwoParams( FUpdateContentForCategoryDelegate, FName, FText);
 
 	/**
-	 * A delegate which provides the const FName& Command name as a parameter to indicates which category was clicked, and returns
-	 * a TSharedRef<SWidget> meant to populate the full content for that category
+	 * @return the TArray<FName> of Favorites for this 
+	 */
+	WIDGETREGISTRATION_API const TArray<FName>& GetFavorites() const;
+
+	/**
+	 * @return the context menu containing an item to add or remove a favorite. If FavoritesItemName is a
+	 * current favorite it will return a context menu with the proper item to add or remove the favorite
+	 * 
+	 * @param FavoritesItemName the name of the Favorites item that would be toggled with the menu
+	 */
+	WIDGETREGISTRATION_API TSharedRef<SWidget> CreateFavoritesContextMenu( FString FavoritesItemName );
+
+	/**
+	 * A delegate which provides the const FName& Command name as a parameter to indicates which category was clicked.
+	 * The bound method should update this FCategoryDrivenContentBuilder's content based on the category that was chosen.
 	 */ 
-	FProvideSelectedCategoryContent  ProvideSelectedCategoryContentDelegate;
+	FUpdateContentForCategoryDelegate UpdateContentForCategoryDelegate;
 	
 public:
 	/**
@@ -30,48 +60,116 @@ public:
 	/**
 	 * destroys the FCategoryDrivenContentBuilder and frees any resources
 	 */
-	WIDGETREGISTRATION_API virtual ~FCategoryDrivenContentBuilder() override;
+	WIDGETREGISTRATION_API ~FCategoryDrivenContentBuilder();
 
-private:
-	/**
-	 * Given the active category name, update the  content
-	 * 
-	 * @param ActiveCategoryName 
-	 */
-	WIDGETREGISTRATION_API virtual void ProvideSelectedCategoryContent( FName ActiveCategoryName = NAME_None ) override;
-
-	FName ActiveCommandName;
-
-public:
 	/*
 	 * Initializes the data necessary to build the category toolbar
 	 */
 	WIDGETREGISTRATION_API virtual void InitializeCategoryToolbar() override;
 
+	/**
+	 * Generate the widget built by this builder and return a shared pointer to it.
+	 */
+	virtual TSharedPtr<SWidget> GenerateWidget() override;
+
 	/*
 	 * Refreshes the existing widget.
 	 */
 	WIDGETREGISTRATION_API virtual void UpdateWidget() override;
 
+
 	/**
-	 * Returns true if the FUICommandInfo with the name CommandName is the active tool palette,
-	 * else it returns false
+	 * Sets whether we should show that no category selected. If true, the category picker shows no selection and there is no category title.
 	 *
-	 * @param CommandName the name of the FUICommandInfo we are checking to see if it is the active tool palette
+	 * @param bInShowNoCategorySelection if true, the category picker will show no selection and the category title is not shown
 	 */
-	WIDGETREGISTRATION_API virtual ECheckBoxState IsActiveToolPalette(FName CommandName) const override
-	{
-		return ActiveCommandName == CommandName ? ECheckBoxState::Checked : ECheckBoxState::Unchecked;
-	}
+	WIDGETREGISTRATION_API void SetShowNoCategorySelection( bool bInShowNoCategorySelection );
+	
+	WIDGETREGISTRATION_API void SetCommands(TArray<TSharedPtr<FUICommandInfo>> InContentLoaderCommands);
 
 	/*
-	 * Sets the commands that will load the categories
+	 * Initializes the category buttons with the given FBuilderInputs
+	 *
+	 * @param InBuilderInputArray the array of FBuilderInput instances that will initialize the Category buttons
 	 */
-	WIDGETREGISTRATION_API void SetCommands(TArray<TSharedPtr<FUICommandInfo>> InContentLoaderCommands);
+	WIDGETREGISTRATION_API void InitializeCategoryButtons(TArray<UE::DisplayBuilders::FBuilderInput> InBuilderInputArray);
+
+	/**
+	 * Converts the SWidget Widget to a FSlateBuilder and adds it to the main content for the currently selected category
+	 *
+	 * @param Widget the SWidget to convert to an FSlateBuilder and add to the content for the current category
+	 */
+	WIDGETREGISTRATION_API void AddBuilder( TSharedRef<SWidget> Widget );
+
+	/**
+	 * Converts the SWidget Widget to a FSlateBuilder and sets it as the entire main content for the currently selected category
+	 *
+	 * @param Widget the SWidget to convert to an FSlateBuilder and set as the content for the current category
+	 */
+	WIDGETREGISTRATION_API void FillWithBuilder( TSharedRef<SWidget> Widget );
+
+	/**
+	 * Clears the content for the currently selected category
+	 */
+	WIDGETREGISTRATION_API void ClearCategoryContent();
 	
 private:
-	TArray<TSharedPtr<FUICommandInfo>> ContentLoaderCommands;
-};
+	/**
+     * Given the active category name, return an SWidget to provide the content
+     * 
+     * @param ActiveCategoryName the name of the Category that will be loaded, which will be the title in the header of the returned content
+     */
+	WIDGETREGISTRATION_API virtual void UpdateContentForCategory(
+		FName ActiveCategoryName = NAME_None
+		, FText InActiveCategoryText = FText::GetEmpty() ) override;
+
+	/**
+	* A map of the Category FName to the TSharedRef<FBuilderInput> providing the information to instantiate the category buttons
+	*/
+	TMap<FName, UE::DisplayBuilders::FBuilderInput> CategoryNameToBuilderInputMap;
+
+	/**
+	 * The TArray of FNames which have been set as favorites
+	 */
+	TArray<FName> Favorites;
+
+	/**
+	 * the FName of the favorites Category
+	 */
+	FName FavoritesCategoryName;
+
+	/**
+	 * Toggles the favorite with the name InFavoriteCommandName
+	 */
+	void ToggleFavorite( FName InFavoriteCommandName );
+	
+	/**
+	 * Adds the favorites with the name InFavoriteCommandName to the favorites list. If InFavoriteCommandName is already in
+	 * the Favorites list, this is a no-op
+	 */
+	void AddFavorite( FName InFavoriteCommandName );
 
+private:
+
+	/** The array of FSlateBuilders that will build the content for the currently selected category */
+	TArray<TSharedRef<FSlateBuilder>> ChildBuilderArray;
+
+	/** the FText label of the currently selected category */
+	FText CategoryLabel;
+	
+	/** the simple title container that has a header with the currently selected category name, and an empty body that will contain
+	 * the category content */
+	TSharedPtr<FSimpleTitleContainer> TitleContainer;
 
- 
\ No newline at end of file
+	/** the column wrapping container that provides the body of the category content */
+	TSharedPtr<FColumnWrappingContainer> ColumnWrappingContainer;
+
+	/** if true, the content for the category should have a single widget fill the content */
+	bool bIsFilledWithWidget;
+
+	/** if true, we need to show no selection on the category picker and no category title, else we should show both  */
+	bool bShowNoCategorySelection;
+	
+	/** The dropzone handler for the favorites */
+	// TSharedPtr<FAcceptSingleItemDropZoneHandler> FavoritesDropZoneHandler;
+};
diff --git a/Engine/Source/Developer/WidgetRegistration/Public/Layout/CategoryDrivenContentBuilderBase.h b/Engine/Source/Developer/WidgetRegistration/Public/Layout/CategoryDrivenContentBuilderBase.h
--- a/Engine/Source/Developer/WidgetRegistration/Public/Layout/CategoryDrivenContentBuilderBase.h
+++ b/Engine/Source/Developer/WidgetRegistration/Public/Layout/CategoryDrivenContentBuilderBase.h
@@ -5,18 +5,21 @@
 
 #include "FToolkitWidgetStyle.h"
 #include "ToolElementRegistry.h"
+#include "Widgets/SBoxPanel.h"
+#include "Widgets/Layout/SSplitter.h"
 #include "Framework/MultiBox/MultiBoxBuilder.h"
-
+#include "BuilderKey.h"
 
 class SWidget;
 class FToolElement;
 
-struct FCategoryDrivenContentBuilderArgs;
+class FCategoryDrivenContentBuilderArgs;
+// class FAcceptSingleItemDropZoneHandler;
 
 /**
  * The FToolElementRegistrationArgs which is specified for Toolkits
  */
-class WIDGETREGISTRATION_API FCategoryDrivenContentBuilderBase : public FToolElementRegistrationArgs
+class FCategoryDrivenContentBuilderBase : public FToolElementRegistrationArgs
 {
 public:
 
@@ -38,7 +41,7 @@ public:
 	 * A constructor for the FCategoryDrivenContentBuilderBase which defines the builder name
 	 * @param InBuilderName the name for the builder
 	 */
-	explicit  FCategoryDrivenContentBuilderBase( FName InBuilderName );
+	explicit  FCategoryDrivenContentBuilderBase( FName InBuilderName = "FCategoryDrivenContentBuilderBase" );
 
 	/**
 	 * A constructor for the FCategoryDrivenContentBuilderBase which takes a
@@ -68,29 +71,29 @@ public:
 	 * @param Visibility If Visibility == EVisibility::Collapsed, the category button labels
 	 * will be shown, else they will not be shown.
 	 */
-	void SetCategoryButtonLabelVisibility(EVisibility Visibility);
+	WIDGETREGISTRATION_API void SetCategoryButtonLabelVisibility(EVisibility Visibility);
 
 	/**
-	 * Sets category button label visibility to Visiblity. It also reinitializes the
+	 * Sets category button label visibility to Visibility. It also reinitializes the
 	 * category toolbar data, as the toolbar's LabelVisibility member is now stale.
 	 *
 	 * @param bIsCategoryButtonLabelVisible If bIsCategoryButtonLabelVisible == true,
 	 * the category button labels will be shown, else they will not be shown.
 	 */
-	void SetCategoryButtonLabelVisibility(bool bIsCategoryButtonLabelVisible);
+	WIDGETREGISTRATION_API void SetCategoryButtonLabelVisibility(bool bIsCategoryButtonLabelVisible);
 
 	/**
-	 * RefreshCategoryToolbarWidget refreshes the UI display of the category toolbar 
+	 * refreshes the UI display of the category toolbar 
 	 */
-	void RefreshCategoryToolbarWidget();
+	WIDGETREGISTRATION_API void RefreshCategoryToolbarWidget();
 
 	/** Implements the generation of the TSharedPtr<SWidget> */
-	virtual TSharedPtr<SWidget> GenerateWidget() override;
+	WIDGETREGISTRATION_API virtual TSharedPtr<SWidget> GenerateWidget() override;
 
 	/** Creates the Toolbar for the widget with the FUICommandInfos that load the Palettes */
 	TSharedRef<SWidget> CreateToolbarWidget() const;
 
-	/** returns a TSharedPointer to the FToolbarBuilder with the FUICommandInfos that load the Palettes */
+	/** @return a TSharedPointer to the FToolbarBuilder with the FUICommandInfos that load the Palettes */
 	TSharedPtr<FToolBarBuilder> GetLoadPaletteToolbar();
 	
 	/** a TSharedPointer to the FToolElement for the vertical toolbar */
@@ -104,14 +107,6 @@ private:
 	/* The SWidget that is the whole Toolkit */
 	TSharedPtr<SWidget> ToolkitWidget;
 
-	/**
-	 * Returns true if the FUICommandInfo with the name CommandName is the active tool palette,
-	 * else it returns false
-	 *
-	 * @param CommandName the name of the FUICommandInfo we are checking to see if it is the active tool palette
-	 */
-	virtual ECheckBoxState IsActiveToolPalette(FName CommandName) const = 0;
-
 	/**
 	 * Creates the SWidget
 	 */
@@ -124,15 +119,15 @@ private:
 	TSharedPtr<SVerticalBox> ToolkitWidgetContainerVBox;
 
 protected:
-	virtual void ProvideSelectedCategoryContent( FName ActiveCategoryName = NAME_None) = 0;
+	virtual void UpdateContentForCategory( FName ActiveCategoryName = NAME_None, FText InActiveCategoryText = FText::GetEmpty() ) = 0;
 	
-	/** The SVerticalBox which holds all but the vertical toolbar in a Toolkit */
-	TSharedPtr<SVerticalBox> ToolkitWidgetVBox;
+	/** The SVerticalBox which holds the main content ~ -all but the Category chooser */
+	TSharedPtr<SVerticalBox> MainContentVerticalBox;
 	
 	/** the tool element registry this class will use to register UI tool elements */
 	static FToolElementRegistry ToolRegistry;
 	
-	/** Specifies what happens if you click the category button of an already-active catogory */
+	/** Specifies what happens if you click the category button of an already-active category */
 	ECategoryReclickBehavior CategoryReclickBehavior = ECategoryReclickBehavior::NoEffect;
 	
 	/** The current FToolkitWidgetStyle */
@@ -156,15 +151,30 @@ protected:
 
 	/** Name of this builder */
 	FName BuilderName;
+	
+	/**
+	 * the name of the currently selected/active category
+	 */
+	FName ActiveCategoryName;
 };
 
 /**
  * A struct to provide arguments for a FCategoryDrivenContentBuilderBase
  */
-struct FCategoryDrivenContentBuilderArgs
+class FCategoryDrivenContentBuilderArgs
 {
 public:
-	WIDGETREGISTRATION_API FCategoryDrivenContentBuilderArgs( FName InBuilderName );
+	/**
+	* Constructor taking the builder name and the FBuilderKey
+	*
+	* @param InBuilderName the name of the builder
+	* @param InKey the FBuilderKey for this 
+	*/
+	WIDGETREGISTRATION_API FCategoryDrivenContentBuilderArgs(
+		FName InBuilderName, const UE::DisplayBuilders:: FBuilderKey& InKey = UE::DisplayBuilders::FBuilderKeys::Get().None() );
+
+	/** the FBuilderKey for this */
+	UE::DisplayBuilders::FBuilderKey Key;
 
 	/** Name of this builder */
 	FName BuilderName;
@@ -175,4 +185,16 @@ public:
 
 	/** Specifies what happens if you click the category button of an already-active category */
 	FCategoryDrivenContentBuilderBase::ECategoryReclickBehavior CategoryReclickBehavior;
+
+	/** The FName of the favorites Category, if one exists */
+	FName FavoritesCommandName;
+
+	/** The label/title of the initially selected category  */
+	FText CategoryLabel;
+	
+	/** The name of the initially selected category */
+	FName ActiveCategoryName;
+	
+	
+	// TSharedPtr<FAcceptSingleItemDropZoneHandler> FavoritesDropZoneHandler;
 };
\ No newline at end of file
diff --git a/Engine/Source/Developer/WidgetRegistration/Public/Persistence/BuilderPersistenceManager.h b/Engine/Source/Developer/WidgetRegistration/Public/Persistence/BuilderPersistenceManager.h
--- a/Engine/Source/Developer/WidgetRegistration/Public/Persistence/BuilderPersistenceManager.h
+++ b/Engine/Source/Developer/WidgetRegistration/Public/Persistence/BuilderPersistenceManager.h
@@ -33,7 +33,7 @@ class WIDGETREGISTRATION_API  UBuilderPersistenceManager : public UEditorConfigB
 
 public:
 	/**
-	 * Returns the default named favorites array for the Builder with FBuilderKey Key, if one exists, else an empty array is returned
+	 * @return the default named favorites array for the Builder with FBuilderKey Key, if one exists, else an empty array is returned
 	 * 
 	 * @param Key the FBuilderKey to retrieve the favorites for
 	 */
@@ -48,7 +48,7 @@ public:
 	void SetPersistedFavoritesNamesArray( const UE::DisplayBuilders::FBuilderKey& Key, TArray<FName>& Favorites );
 
 	/**
-	 * Returns the an array for the Builder with FBuilderKey Key and the suffix that the array was persisted with, if one exists, else an empty array is returned
+	 * @return the an array for the Builder with FBuilderKey Key and the suffix that the array was persisted with, if one exists, else an empty array is returned
 	 * 
 	 * @param Key the FBuilderKey to retrieve the FNames array for
 	 * @param PersistenceKeySuffix the suffix to add to the FBuilderKey to persist the array of FNames with
diff --git a/Engine/Source/Developer/WidgetRegistration/Public/ToolkitBuilder.h b/Engine/Source/Developer/WidgetRegistration/Public/ToolkitBuilder.h
--- a/Engine/Source/Developer/WidgetRegistration/Public/ToolkitBuilder.h
+++ b/Engine/Source/Developer/WidgetRegistration/Public/ToolkitBuilder.h
@@ -131,13 +131,13 @@ public:
 	 * @param InToolkitCommandList  the toolkit FUICommandList
 	 * @param InToolkitSections The FToolkitSections for this toolkit builder
 	 */
-	FToolkitBuilder(
+	explicit FToolkitBuilder(
 	FName ToolbarCustomizationName,
 	TSharedPtr<FUICommandList> InToolkitCommandList,
 	TSharedPtr<FToolkitSections> InToolkitSections);
 
 	/** default constructor */
-	FToolkitBuilder(FToolkitBuilderArgs& Args);
+	explicit FToolkitBuilder(FToolkitBuilderArgs& Args);
 
 	virtual ~FToolkitBuilder() override;
 
@@ -169,7 +169,7 @@ public:
 	 * will be regenerated to reflect it.  */
 	virtual void UpdateWidget() override;
 
-	/** returns true is there is an active palette selected, else it returns false */
+	/** @return true is there is an active palette selected, else it returns false */
 	bool HasActivePalette() const;
 
 	/*
@@ -180,7 +180,7 @@ public:
 	void SetActivePaletteOnLoad(const FUICommandInfo* Command);
 	
 	/**
-	 * Returns true if the FUICommandInfo with the name CommandName is the active tool palette,
+	 * @return true if the FUICommandInfo with the name CommandName is the active tool palette,
 	 * else it returns false
 	 *
 	 * @param CommandName the name of the FUICommandInfo we are checking to see if it is the active tool palette
@@ -290,7 +290,7 @@ private:
 	 */
 	TSharedRef<SWidget> GetContextMenuContent(const FName CommandName);
 	
-	/** returns a TSharedPointer to the FToolbarBuilder with the FUICommandInfos that load the Palettes */
+	/** @return a TSharedPointer to the FToolbarBuilder with the FUICommandInfos that load the Palettes */
 	TSharedRef<SWidget> GetToolPaletteWidget() const;
 	
 	void CreatePaletteWidget(FToolPalette& Palette, FToolElement& Element);
@@ -304,7 +304,7 @@ private:
 	 */
 	void ToggleCommandInPalette(TSharedRef<FEditablePalette> Palette, FString CommandNameString);
 		
-	/* Returns true if the Toolkit builder has some tools that are currently active/selected */
+	/* @return true if the Toolkit builder has some tools that are currently active/selected */
 	bool HasSelectedToolSet() const;
 
 	/**
@@ -312,7 +312,7 @@ private:
 	 * 
 	 * @param ActiveCategoryName 
 	 */
-	virtual void ProvideSelectedCategoryContent( FName ActiveCategoryName = NAME_None ) override;
+	virtual void UpdateContentForCategory( FName ActiveCategoryName = NAME_None, FText InActiveCategoryText = FText::GetEmpty() ) override;
 
 	/**
 	 * Gets the EVisibility of the active tool title. This should only be visible if a tool is currently chosen in the palette.
diff --git a/Engine/Source/Editor/EditorStyle/Private/StarshipStyle.cpp b/Engine/Source/Editor/EditorStyle/Private/StarshipStyle.cpp
--- a/Engine/Source/Editor/EditorStyle/Private/StarshipStyle.cpp
+++ b/Engine/Source/Editor/EditorStyle/Private/StarshipStyle.cpp
@@ -1737,7 +1737,7 @@ void FStarshipEditorStyle::FStyle::SetupGeneralStyles()
 		FVerticalToolBarStyle.SetBackgroundPadding(   FMargin(6.0f, 4.0f));
 		FVerticalToolBarStyle.ButtonStyle.SetNormalPadding(FMargin(12.0f, 6.0f));
 		FVerticalToolBarStyle.ButtonStyle.SetPressedPadding(FMargin(12.0f, 6.0f));
-		FVerticalToolBarStyle.SetExpandBrush(IMAGE_BRUSH("Icons/toolbar_expand_16x", Icon8x8));
+		FVerticalToolBarStyle.SetExpandBrush(CORE_IMAGE_BRUSH_SVG("Starship/Common/ellipsis-horizontal-narrow", FVector2D(16, 16)));
 		FVerticalToolBarStyle.SetIconPadding(FMargin(8.0f, 8.0f));
 		FVerticalToolBarStyle.SetIconPaddingWithVisibleLabel(FMargin(8.0f, 8.0f, 8.0f, 2.0f));
  
@@ -1759,6 +1759,14 @@ void FStarshipEditorStyle::FStyle::SetupGeneralStyles()
 		FVerticalToolBarStyle.SetToggleButtonStyle(CheckBoxStyle);
 		Set("FVerticalToolBar", FVerticalToolBarStyle);
 	}
+
+	{
+		// FCategoryDrivenContentBuilder vertical toolbar style
+		FToolBarStyle CategoryDrivenContentBuilderToolBarStyle = GetWidgetStyle<FToolBarStyle>("FVerticalToolBar");
+		CategoryDrivenContentBuilderToolBarStyle.SetIconPaddingWithVisibleLabel(FMargin(8.0f, 12.0f, 8.0f, 8.0f));
+		CategoryDrivenContentBuilderToolBarStyle.SetButtonContentMaxWidth( 56.0f );
+		Set("FCategoryDrivenContentBuilderToolbar", CategoryDrivenContentBuilderToolBarStyle);
+	}
 	
 	// Vertical ToolPalette 
 	{
@@ -4481,6 +4489,7 @@ void FStarshipEditorStyle::FStyle::SetupLevelEditorStyle()
 
 		Set( "PlacementBrowser.Asset.Background", new FSlateRoundedBoxBrush(FStyleColors::Recessed, 6.f));
 		Set( "PlacementBrowser.Asset.LabelBack", new BOX_BRUSH("Starship/PlacementBrowser/LabelBack_18x", 6.f/18.f, FStyleColors::Dropdown));
+		Set( "PlacementBrowser.Asset.ThumbnailBackground", new FSlateRoundedBoxBrush(FStyleColors::Dropdown, 6 ));
 
 		FLinearColor DimBackground = FLinearColor( FColor( 64, 64, 64 ) );
 		FLinearColor DimBackgroundHover = FLinearColor( FColor( 50, 50, 50 ) );
