Fixing  editor failure, by checking for instance before removing from root

diff --git a/Engine/Source/Developer/WidgetRegistration/Private/Persistence/BuilderPersistenceManager.cpp b/Engine/Source/Developer/WidgetRegistration/Private/Persistence/BuilderPersistenceManager.cpp
--- a/Engine/Source/Developer/WidgetRegistration/Private/Persistence/BuilderPersistenceManager.cpp
+++ b/Engine/Source/Developer/WidgetRegistration/Private/Persistence/BuilderPersistenceManager.cpp
@@ -18,7 +18,10 @@ void UBuilderPersistenceManager::Initialize()
 
 void UBuilderPersistenceManager::ShutDown()
 {
-	Instance->RemoveFromRoot();
+	if ( UObjectInitialized() )
+	{
+		Instance->RemoveFromRoot();
+	}
 	Instance = nullptr;
 }
 
