Place Actors: Add a context menu to hide and show labels and persist the user choice

diff --git a/Engine/Source/Developer/WidgetRegistration/Private/Layout/CategoryDrivenContentBuilder.cpp b/Engine/Source/Developer/WidgetRegistration/Private/Layout/CategoryDrivenContentBuilder.cpp
--- a/Engine/Source/Developer/WidgetRegistration/Private/Layout/CategoryDrivenContentBuilder.cpp
+++ b/Engine/Source/Developer/WidgetRegistration/Private/Layout/CategoryDrivenContentBuilder.cpp
@@ -30,8 +30,6 @@ TSharedRef<SWidget> FCategoryDrivenContentBuilder::CreateFavoritesContextMenu( F
 	
 	FMenuBuilder MenuBuilder(bInShouldCloseWindowAfterMenuSelection, InCommandList);
 
-	UBuilderPersistenceManager::Get()->PersistArrayOfNames( BuilderKey, BuilderKey.ToName(), Favorites );
-
 	const FText ItemText = Favorites.Contains( FavoritesItemName ) ?
 		LOCTEXT("CategoryDrivenContentBuilder_RemoveFromFavorites", "Remove from Favorites") :
 		LOCTEXT("CategoryDrivenContentBuilder_AddToFavorites", "Add to Favorites");
@@ -42,6 +40,23 @@ TSharedRef<SWidget> FCategoryDrivenContentBuilder::CreateFavoritesContextMenu( F
 	return MenuBuilder.MakeWidget();
 }
 
+TSharedRef<SWidget> FCategoryDrivenContentBuilder::CreateShowCategoryLabelsContextMenu()
+{
+	const bool bInShouldCloseWindowAfterMenuSelection = true;
+	const TSharedPtr< const FUICommandList > CommandList = nullptr;
+	
+	FMenuBuilder MenuBuilder(bInShouldCloseWindowAfterMenuSelection, CommandList);
+
+	const FText ItemText = CategoryButtonLabelVisibility.IsVisible() ?
+		LOCTEXT("CategoryDrivenContentBuilder_HideLabels", "Hide Labels") :
+		LOCTEXT("CategoryDrivenContentBuilder_ShowLabels", "Show Labels");
+		
+	const FUIAction ItemAction(FExecuteAction::CreateSP(this, &FCategoryDrivenContentBuilder::ToggleShowLabels));
+	MenuBuilder.AddMenuEntry(ItemText, ItemText, FSlateIcon(), ItemAction);
+		
+	return MenuBuilder.MakeWidget();
+}
+
 FCategoryDrivenContentBuilder::FCategoryDrivenContentBuilder(FCategoryDrivenContentBuilderArgs& Args):
 	FCategoryDrivenContentBuilderBase(Args )
 	, FavoritesCategoryName( Args.FavoritesCommandName )
@@ -50,7 +65,9 @@ FCategoryDrivenContentBuilder::FCategoryDrivenContentBuilder(FCategoryDrivenCont
 	, bIsFilledWithWidget( false )
 	, bShowNoCategorySelection( false )
 {
-	Favorites = UBuilderPersistenceManager::Get()->GetPersistedFavoritesNamesArray( BuilderKey );
+	Favorites = UBuilderPersistenceManager::Get()->GetFavoritesNames( BuilderKey );
+	const bool bCategoryButtonLabelVisibilityBool = UBuilderPersistenceManager::Get()->GetShowButtonLabels( BuilderKey, CategoryButtonLabelVisibility.IsVisible() );
+	CategoryButtonLabelVisibility = bCategoryButtonLabelVisibilityBool ? EVisibility::Visible : EVisibility::Collapsed;
 	ActiveCategoryName = Args.ActiveCategoryName;
 }
 
@@ -80,10 +97,21 @@ void FCategoryDrivenContentBuilder::ToggleFavorite( FName InFavoriteCommandName
 		Favorites.Add( InFavoriteCommandName );
 	}
 
-	UBuilderPersistenceManager::Get()->SetPersistedFavoritesNamesArray( BuilderKey, Favorites );
+	UBuilderPersistenceManager::Get()->PersistFavoritesNames( BuilderKey, Favorites );
 	UpdateWidget();
 }
 
+void FCategoryDrivenContentBuilder::ToggleShowLabels()
+{
+	CategoryButtonLabelVisibility = CategoryButtonLabelVisibility.IsVisible() ? EVisibility::Collapsed : EVisibility::Visible;
+	UBuilderPersistenceManager::Get()->PersistShowButtonLabels( BuilderKey, CategoryButtonLabelVisibility.IsVisible() );
+
+	const bool bShowReinitialize = true;
+	RefreshCategoryToolbarWidget( bShowReinitialize );
+	InitializeCategoryButtons();
+	CreateWidget();
+}
+
 void FCategoryDrivenContentBuilder::AddFavorite(FName InFavoriteCommandName)
 {
 	if ( !Favorites.Contains( InFavoriteCommandName ) )
@@ -113,11 +141,19 @@ void FCategoryDrivenContentBuilder::ClearCategoryContent()
 
 void FCategoryDrivenContentBuilder::InitializeCategoryToolbar()
 {
+	if (!LoadPaletteToolBarBuilder)
+	{
+		return;
+	}
+
+	LoadPaletteToolBarBuilder->SetLabelVisibility(CategoryButtonLabelVisibility);
+	
 	for ( TTuple<FName, UE::DisplayBuilders::FBuilderInput>  Pair : CategoryNameToBuilderInputMap )
 	{
 		FButtonArgs ButtonArgs = Pair.Value.ButtonArgs;
 		const TSharedPtr<const FUICommandInfo> Command = ButtonArgs.Command;
 		const FName CommandName = Command->GetCommandName();
+		ButtonArgs.OnGetMenuContent.BindSP(this, &FCategoryDrivenContentBuilder::CreateShowCategoryLabelsContextMenu );
 		
 		LoadToolPaletteCommandList->MapAction(
 			Command,
diff --git a/Engine/Source/Developer/WidgetRegistration/Private/Persistence/BuilderPersistenceManager.cpp b/Engine/Source/Developer/WidgetRegistration/Private/Persistence/BuilderPersistenceManager.cpp
--- a/Engine/Source/Developer/WidgetRegistration/Private/Persistence/BuilderPersistenceManager.cpp
+++ b/Engine/Source/Developer/WidgetRegistration/Private/Persistence/BuilderPersistenceManager.cpp
@@ -16,6 +16,8 @@ void UBuilderPersistenceManager::Initialize()
 	{
 		Instance = NewObject<UBuilderPersistenceManager>(); 
 		Instance->AddToRoot();
+
+		Instance->LoadEditorConfig();
 	}
 }
 
@@ -28,19 +30,28 @@ void UBuilderPersistenceManager::ShutDown()
 	Instance = nullptr;
 }
 
-TArray<FName> UBuilderPersistenceManager::GetPersistedFavoritesNamesArray(const UE::DisplayBuilders::FBuilderKey& Key)
+TArray<FName> UBuilderPersistenceManager::GetFavoritesNames(const UE::DisplayBuilders::FBuilderKey& Key)
 {
 	return GetPersistedArrayOfNames( Key, UE::DisplayBuilders::BuilderPersistenceManager::FavoritesSuffix );
 }
 
-void UBuilderPersistenceManager::SetPersistedFavoritesNamesArray(const UE::DisplayBuilders::FBuilderKey& Key, TArray<FName>& Favorites)
+void UBuilderPersistenceManager::PersistFavoritesNames(const UE::DisplayBuilders::FBuilderKey& Key, TArray<FName>& Favorites)
 {
 	return PersistArrayOfNames(Key, UE::DisplayBuilders::BuilderPersistenceManager::FavoritesSuffix, Favorites);
 }
 
+bool UBuilderPersistenceManager::GetShowButtonLabels( const UE::DisplayBuilders::FBuilderKey& Key, bool bDefaultValue )
+{
+	return GetPersistedBool( Key, UE::DisplayBuilders::BuilderPersistenceManager::ShowButtonLabelsSuffix, bDefaultValue );
+}
+
+void UBuilderPersistenceManager::PersistShowButtonLabels( const UE::DisplayBuilders::FBuilderKey& Key, bool bValue )
+{
+	return PersistBool( Key, UE::DisplayBuilders::BuilderPersistenceManager::ShowButtonLabelsSuffix, bValue );
+}
+
 TArray<FName> UBuilderPersistenceManager::GetPersistedArrayOfNames( const UE::DisplayBuilders::FBuilderKey& Key, FName PersistenceKeySuffix )
 {
-	LoadEditorConfig();
 	if ( !Key.IsNone()  && !PersistenceKeySuffix.IsNone() )
 	{
 		if ( const FPersistedNameArray* Settings =  SavedNameToPersistedFNameArrayMap.Find( Key.GetKeyWithSuffix( PersistenceKeySuffix ) ) )
@@ -65,35 +76,20 @@ void UBuilderPersistenceManager::PersistArrayOfNames( const UE::DisplayBuilders:
 	SaveEditorConfig();
 }
 
-bool UBuilderPersistenceManager::GetPersistedBool( const UE::DisplayBuilders::FBuilderKey& Key, FName PersistenceKeySuffix, bool PersistedBoolIfNoneFound )
+bool UBuilderPersistenceManager::GetPersistedBool( const UE::DisplayBuilders::FBuilderKey& Key, FName PersistenceKeySuffix, bool bDefaultValue )
 {
-	LoadEditorConfig();
 	if ( !Key.IsNone()  && !PersistenceKeySuffix.IsNone() )
 	{
 		if ( const FPersistedBool* BoolSettings =  SavedNameToPersistedBoolMap.Find( Key.GetKeyWithSuffix( PersistenceKeySuffix ) ) )
 		{
 			return BoolSettings->PersistedBool;
 		}
-		else
-		{
-			PersistBool( Key, PersistenceKeySuffix, PersistedBoolIfNoneFound );
-		}
 	}
 
-	return PersistedBoolIfNoneFound;
-}
-
-void UBuilderPersistenceManager::SetPersistedButtonLabelBool( const UE::DisplayBuilders::FBuilderKey& Key, bool InPersistedBool )
-{
-	return PersistBool( Key, UE::DisplayBuilders::BuilderPersistenceManager::ShowButtonLabelsSuffix, InPersistedBool );
-}
-
-bool UBuilderPersistenceManager::GetPersistedButtonLabelBool( const UE::DisplayBuilders::FBuilderKey& Key, bool PersistedBoolIfNoneFound )
-{
-	return GetPersistedBool( Key, UE::DisplayBuilders::BuilderPersistenceManager::ShowButtonLabelsSuffix, PersistedBoolIfNoneFound );
+	return bDefaultValue;
 }
 
-void UBuilderPersistenceManager::PersistBool( const UE::DisplayBuilders::FBuilderKey& Key, FName PersistenceKeySuffix, bool InPersistedBool )
+void UBuilderPersistenceManager::PersistBool( const UE::DisplayBuilders::FBuilderKey& Key, FName PersistenceKeySuffix, bool bValue )
 {
 	if (PersistenceKeySuffix.IsNone())
 	{
@@ -101,6 +97,6 @@ void UBuilderPersistenceManager::PersistBool( const UE::DisplayBuilders::FBuilde
 	}
 
 	FPersistedBool& Settings =  SavedNameToPersistedBoolMap.Add( Key.GetKeyWithSuffix( PersistenceKeySuffix ) );
-	Settings.PersistedBool = InPersistedBool;;
+	Settings.PersistedBool = bValue;
 	SaveEditorConfig();
 }
diff --git a/Engine/Source/Developer/WidgetRegistration/Public/Layout/CategoryDrivenContentBuilder.h b/Engine/Source/Developer/WidgetRegistration/Public/Layout/CategoryDrivenContentBuilder.h
--- a/Engine/Source/Developer/WidgetRegistration/Public/Layout/CategoryDrivenContentBuilder.h
+++ b/Engine/Source/Developer/WidgetRegistration/Public/Layout/CategoryDrivenContentBuilder.h
@@ -43,6 +43,13 @@ public:
 	 */
 	WIDGETREGISTRATION_API TSharedRef<SWidget> CreateFavoritesContextMenu( FString FavoritesItemName );
 
+	/**
+	 * Returns the context menu containing an item to show or hide category labels.
+	 * 
+	 * @param bInShowCategoryLabelsItemName the name of the Show Category Labels item that would be toggled with the menu
+	 */
+	WIDGETREGISTRATION_API TSharedRef<SWidget> CreateShowCategoryLabelsContextMenu();
+
 	/**
 	 * A delegate which provides the const FName& Command name as a parameter to indicates which category was clicked.
 	 * The bound method should update this FCategoryDrivenContentBuilder's content based on the category that was chosen.
@@ -127,9 +134,9 @@ private:
      * 
      * @param ActiveCategoryName the name of the Category that will be loaded, which will be the title in the header of the returned content
      */
-	WIDGETREGISTRATION_API virtual void UpdateContentForCategory(
-		FName ActiveCategoryName = NAME_None
-		, FText InActiveCategoryText = FText::GetEmpty() ) override;
+	virtual void UpdateContentForCategory(
+		FName ActiveCategoryName = NAME_None,
+		FText InActiveCategoryText = FText::GetEmpty() ) override;
 
 	/**
 	* A map of the Category FName to the TSharedRef<FBuilderInput> providing the information to instantiate the category buttons
@@ -151,6 +158,11 @@ private:
 	 */
 	void ToggleFavorite( FName InFavoriteCommandName );
 
+	/**
+	 * Toggles the visibility of the categories names
+	 */
+	void ToggleShowLabels();
+
 private:
 
 	/** The array of BuilderInputs, in the correct display order */
diff --git a/Engine/Source/Developer/WidgetRegistration/Public/Persistence/BuilderPersistenceManager.h b/Engine/Source/Developer/WidgetRegistration/Public/Persistence/BuilderPersistenceManager.h
--- a/Engine/Source/Developer/WidgetRegistration/Public/Persistence/BuilderPersistenceManager.h
+++ b/Engine/Source/Developer/WidgetRegistration/Public/Persistence/BuilderPersistenceManager.h
@@ -50,7 +50,7 @@ public:
 	 * 
 	 * @param Key the FBuilderKey to retrieve the favorites for
 	 */
-	TArray<FName> GetPersistedFavoritesNamesArray( const UE::DisplayBuilders::FBuilderKey& Key  );
+	TArray<FName> GetFavoritesNames( const UE::DisplayBuilders::FBuilderKey& Key );
 
 	/**
 	 * Sets the default named favorites array for the Builder with FBuilderKey Key
@@ -58,23 +58,23 @@ public:
 	 * @param Key the FBuilderKey to retrieve the favorites for
 	 * @param Favorites the array of FNames for the favorites to persist
 	 */
-	void SetPersistedFavoritesNamesArray( const UE::DisplayBuilders::FBuilderKey& Key, TArray<FName>& Favorites );
+	void PersistFavoritesNames( const UE::DisplayBuilders::FBuilderKey& Key, TArray<FName>& Favorites );
 
 	/**
-	 * @return the an array for the Builder with FBuilderKey Key and the suffix that the array was persisted with, if one exists, else an empty array is returned
+	 * Sets the default button label EVisisbility for the Builder with FBuilderKey Key
 	 * 
-	 * @param Key the FBuilderKey to retrieve the FNames array for
-	 * @param PersistenceKeySuffix the suffix to add to the FBuilderKey to persist the array of FNames with
+	 * @param Key the FBuilderKey to retrieve the button label Bool for
+	 * @param Bool the bool for the button labels to persist
 	 */
-	TArray<FName> GetPersistedArrayOfNames( const UE::DisplayBuilders::FBuilderKey& Key, FName PersistenceKeySuffix );
-
+	void PersistShowButtonLabels( const UE::DisplayBuilders::FBuilderKey& Key, bool bValue );
+	
 	/**
-	 * Persists the array of FNames the Builder with FBuilderKey Key and suffix PersistenceKeySuffix
+	 * @return the bool or the buttons labels for the Builder with FBuilderKey Key, if one exists, else an empty array is returned
 	 * 
-	 * @param Key the FBuilderKey to retrieve the favorites for
-	 * @param PersistenceKeySuffix the suffix added to the FBuilderKey to persist the array of FNames with
+	 * @param Key the FBuilderKey to retrieve the button label bool
+	 * @param PersistedBoolIfNoneFound the Bool to set as the new persistence value, if one is not found
 	 */
-	void PersistArrayOfNames( const UE::DisplayBuilders::FBuilderKey& Key, FName PersistenceKeySuffix, TArray<FName>& ArrayOfNamesToPersist );
+	bool GetShowButtonLabels( const UE::DisplayBuilders::FBuilderKey& Key, bool bDefaultValue );
 
 	/**
 	 * Initialize the Persistence manager
@@ -107,29 +107,29 @@ private:
 
 private:
 	/**
-	 * Sets the default button label EVisisbility for the Builder with FBuilderKey Key
+	 * @return the an array for the Builder with FBuilderKey Key and the suffix that the array was persisted with, if one exists, else an empty array is returned
 	 * 
-	 * @param Key the FBuilderKey to retrieve the button label Bool for
-	 * @param Bool the bool for the button labels to persist
+	 * @param Key the FBuilderKey to retrieve the FNames array for
+	 * @param PersistenceKeySuffix the suffix to add to the FBuilderKey to persist the array of FNames with
 	 */
-	void SetPersistedButtonLabelBool( const UE::DisplayBuilders::FBuilderKey& Key, bool Bool );
-	
+	TArray<FName> GetPersistedArrayOfNames( const UE::DisplayBuilders::FBuilderKey& Key, FName PersistenceKeySuffix );
+
 	/**
-	 * @return the bool or the buttons labels for the Builder with FBuilderKey Key, if one exists, else an empty array is returned
+	 * Persists the array of FNames the Builder with FBuilderKey Key and suffix PersistenceKeySuffix
 	 * 
-	 * @param Key the FBuilderKey to retrieve the button label bool
-	 * @param PersistedBoolIfNoneFound the Bool to set as the new persistence value, if one is not found
+	 * @param Key the FBuilderKey to retrieve the favorites for
+	 * @param PersistenceKeySuffix the suffix added to the FBuilderKey to persist the array of FNames with
 	 */
-	bool GetPersistedButtonLabelBool( const UE::DisplayBuilders::FBuilderKey& Key, bool PersistedBoolIfNoneFound  );
+	void PersistArrayOfNames( const UE::DisplayBuilders::FBuilderKey& Key, FName PersistenceKeySuffix, TArray<FName>& ArrayOfNamesToPersist );
 
 	/**
 	 * Returns the bool for the Builder with FBuilderKey Key and the suffix that the array was persisted with, if one exists, else an empty array is returned
 	 * 
 	 * @param Key the FBuilderKey to retrieve the bool
 	 * @param PersistenceKeySuffix the suffix to add to the FBuilderKey to get the bool with
-	 * @param PersistedBoolIfNoneFound the Bool to set as the new persistence value, if one is not found
+	 * @param bDefaultValue the Bool to set as the new persistence value, if one is not found
 	 */
-	bool GetPersistedBool( const UE::DisplayBuilders::FBuilderKey& Key, FName PersistenceKeySuffix, bool PersistedBoolIfNoneFound );
+	bool GetPersistedBool( const UE::DisplayBuilders::FBuilderKey& Key, FName PersistenceKeySuffix, bool bDefaultValue );
 	
 	/**
 	 * Persists the Bool the Builder with FBuilderKey Key and suffix PersistenceKeySuffix
@@ -137,5 +137,5 @@ private:
 	 * @param Key the FBuilderKey to retrieve the favorites for
 	 * @param PersistenceKeySuffix the suffix added to the FBuilderKey to persist the Bool with
 	 */
-	void PersistBool( const UE::DisplayBuilders::FBuilderKey& Key, FName PersistenceKeySuffix, bool InPersistedBool );
+	void PersistBool( const UE::DisplayBuilders::FBuilderKey& Key, FName PersistenceKeySuffix, bool bInPersistedBool );
 };
\ No newline at end of file
