WidgetRegistration: Fixed issue where CategoryDrivenContentBuilderBase did not update visibility on reinitialization.

This addresses an issue in FToolkitBuilder, where calling UpdateWidget would update the CategoryToolbarVisibility but would never take effect because the setting was only applied on initial widget construction and calls to reinitialize would reuse that container widget.

diff --git a/Engine/Source/Developer/WidgetRegistration/Private/Layout/CategoryDrivenContentBuilderBase.cpp b/Engine/Source/Developer/WidgetRegistration/Private/Layout/CategoryDrivenContentBuilderBase.cpp
--- a/Engine/Source/Developer/WidgetRegistration/Private/Layout/CategoryDrivenContentBuilderBase.cpp
+++ b/Engine/Source/Developer/WidgetRegistration/Private/Layout/CategoryDrivenContentBuilderBase.cpp
@@ -60,6 +60,7 @@ void FCategoryDrivenContentBuilderBase::InitCategoryToolbarContainerWidget()
 	else
 	{
 		CategoryToolbarVBox->ClearChildren();
+		CategoryToolbarVBox->SetVisibility(CategoryToolbarVisibility);
 	}
 
 	float ToolbarBoxWidth = UE::CategoryDrivenContentBuilderBase::Private::UnlabelledButtonToolbarButtonWidth;
