ToolkitBuilder: Added new Header and ToolHeader toolkit sections

diff --git a/Engine/Source/Developer/WidgetRegistration/Private/ToolkitBuilder.cpp b/Engine/Source/Developer/WidgetRegistration/Private/ToolkitBuilder.cpp
--- a/Engine/Source/Developer/WidgetRegistration/Private/ToolkitBuilder.cpp
+++ b/Engine/Source/Developer/WidgetRegistration/Private/ToolkitBuilder.cpp
@@ -490,6 +490,17 @@ void FToolkitBuilder::UpdateContentForCategory( FName InActiveCategoryName, FTex
 			ToolkitSections->ModeWarningArea->AsShared()
 		];
 	}
+
+	if (ToolkitSections->Header)
+	{
+		MainContentVerticalBox->AddSlot()
+		.AutoHeight()
+		.HAlign(HAlign_Fill)
+		.Padding(0)
+		[
+			ToolkitSections->Header->AsShared()
+		];
+	}
 	
 	MainContentVerticalBox->AddSlot()
 	.AutoHeight()
@@ -557,6 +568,17 @@ void FToolkitBuilder::UpdateContentForCategory( FName InActiveCategoryName, FTex
 		];
 	}
 
+	if (ToolkitSections->ToolHeader)
+	{
+		MainContentVerticalBox->AddSlot()
+		.AutoHeight()
+		.HAlign(HAlign_Fill)
+		.Padding(0)
+		[
+			ToolkitSections->ToolHeader->AsShared()
+		];
+	}
+
 	if (ToolkitSections->DetailsView)
 	{
 		MainContentVerticalBox->AddSlot()
diff --git a/Engine/Source/Developer/WidgetRegistration/Public/ToolkitBuilder.h b/Engine/Source/Developer/WidgetRegistration/Public/ToolkitBuilder.h
--- a/Engine/Source/Developer/WidgetRegistration/Public/ToolkitBuilder.h
+++ b/Engine/Source/Developer/WidgetRegistration/Public/ToolkitBuilder.h
@@ -23,6 +23,8 @@ struct FToolkitSections
 	TSharedPtr<STextBlock> ToolWarningArea = nullptr;
 	TSharedPtr<SWidget> ToolPresetArea = nullptr;
 	TSharedPtr<IDetailsView> DetailsView = nullptr;
+	TSharedPtr<SWidget> ToolHeader = nullptr;
+	TSharedPtr<SWidget> Header = nullptr;
 	TSharedPtr<SWidget> Footer = nullptr;
 };
 
