Allow the tool builder to have only one palette

diff --git a/Engine/Source/Developer/WidgetRegistration/Private/ToolkitBuilder.cpp b/Engine/Source/Developer/WidgetRegistration/Private/ToolkitBuilder.cpp
--- a/Engine/Source/Developer/WidgetRegistration/Private/ToolkitBuilder.cpp
+++ b/Engine/Source/Developer/WidgetRegistration/Private/ToolkitBuilder.cpp
@@ -228,7 +228,7 @@ void FToolkitBuilder::UpdateWidget()
 		UpdateEditablePalette(EditablePalette);		
 	}
 	
-	CategoryToolbarVisibility = LoadCommandNameToToolPaletteMap.Num() > 1 ? EVisibility::Visible : EVisibility::Collapsed;
+	CategoryToolbarVisibility = LoadCommandNameToToolPaletteMap.Num() > 0 ? EVisibility::Visible : EVisibility::Collapsed;
 }
 
 void FToolkitBuilder::ToggleCommandInPalette(TSharedRef<FEditablePalette> Palette, FString CommandNameString)
