Toolkit builder: fixed BodyContentSBox being overriden by a loose SBox

diff --git a/Engine/Source/Developer/WidgetRegistration/Private/Layout/Containers/HeaderAndBodyContainer.cpp b/Engine/Source/Developer/WidgetRegistration/Private/Layout/Containers/HeaderAndBodyContainer.cpp
--- a/Engine/Source/Developer/WidgetRegistration/Private/Layout/Containers/HeaderAndBodyContainer.cpp
+++ b/Engine/Source/Developer/WidgetRegistration/Private/Layout/Containers/HeaderAndBodyContainer.cpp
@@ -121,22 +121,24 @@ TSharedPtr<SWidget> FHeaderAndBodyContainer::GenerateWidget()
 			]
 		]
 		+ SVerticalBox::Slot()
-		.FillHeight(1.0f)[SNew(SBorder)
-		.BorderImage(ContainerBackgroundBrush)
-		.Padding(0.f)
+		.FillHeight(1.0f)
 		[
-			SAssignNew(BodyContentSBox, SBox )
-			.WidthOverride( FOptionalSize() )
-			.HeightOverride( FOptionalSize() )
+			SNew(SBorder)
+			.BorderImage(ContainerBackgroundBrush)
+			.Padding(0.f)
 			[
-				SAssignNew( BodyContentSScrollBox, SScrollBox)
-				+ SScrollBox::Slot()
+				SAssignNew(BodyContentSBox, SBox)
+				.WidthOverride( FOptionalSize() )
+				.HeightOverride( FOptionalSize() )
 				[
-					BodyBuilder->GenerateWidgetSharedRef()
+					SAssignNew( BodyContentSScrollBox, SScrollBox)
+					+ SScrollBox::Slot()
+					[
+						BodyBuilder->GenerateWidgetSharedRef()
+					]
 				]
 			]
-		]
-	];
+		];
 
 	if ( bHasToggleButtonToCollapseBody )
 	{
@@ -196,7 +198,6 @@ void FHeaderAndBodyContainer::UpdateToBodyAddedState()
 {
 	if ( BodyContentSScrollBox.IsValid() )
 	{
-		BodyContentSBox = SNew( SBox );
 		BodyContentSBox->SetContent( BodyContentSScrollBox.ToSharedRef() );
 	
 		if ( ToggleExpansionImage.IsValid() )
