ToolkitBuilder: Expose setter for CategoryToolbarVisibility.

diff --git a/Engine/Source/Developer/WidgetRegistration/Private/ToolkitBuilder.cpp b/Engine/Source/Developer/WidgetRegistration/Private/ToolkitBuilder.cpp
--- a/Engine/Source/Developer/WidgetRegistration/Private/ToolkitBuilder.cpp
+++ b/Engine/Source/Developer/WidgetRegistration/Private/ToolkitBuilder.cpp
@@ -227,8 +227,6 @@ void FToolkitBuilder::UpdateWidget()
 	{
 		UpdateEditablePalette(EditablePalette);		
 	}
-	
-	CategoryToolbarVisibility = LoadCommandNameToToolPaletteMap.Num() > 0 ? EVisibility::Visible : EVisibility::Collapsed;
 }
 
 void FToolkitBuilder::ToggleCommandInPalette(TSharedRef<FEditablePalette> Palette, FString CommandNameString)
@@ -288,6 +286,11 @@ void FToolkitBuilder::InitializeCategoryToolbar(bool bInitLoadToolPaletteMap)
 
 }
 
+void FToolkitBuilder::SetCategoryToolbarVisibility(EVisibility Visibility)
+{
+	CategoryToolbarVisibility = Visibility;
+}
+
 void FToolkitBuilder::TogglePalette(TSharedPtr<FToolPalette> Palette)
 {
 	const FName CommandName = Palette->LoadToolPaletteAction->GetCommandName();
diff --git a/Engine/Source/Developer/WidgetRegistration/Public/Layout/CategoryDrivenContentBuilderBase.h b/Engine/Source/Developer/WidgetRegistration/Public/Layout/CategoryDrivenContentBuilderBase.h
--- a/Engine/Source/Developer/WidgetRegistration/Public/Layout/CategoryDrivenContentBuilderBase.h
+++ b/Engine/Source/Developer/WidgetRegistration/Public/Layout/CategoryDrivenContentBuilderBase.h
@@ -145,10 +145,9 @@ protected:
 	EVisibility CategoryButtonLabelVisibility;
 	
 	/*
-	 * Is the Category toolbar visible? This should be false if the Category toolbar has less than 2 categories,
-	 * because if there is only one showing it is superfluous to actually show it since it will always be selected.
+	 * Is the Category toolbar visible?
 	 */
-	EVisibility CategoryToolbarVisibility;
+	EVisibility CategoryToolbarVisibility = EVisibility::Visible;
 	
 	/** A TSharedPointer to the FUICommandList for the FUICommandInfos which load a tool palette */
 	TSharedPtr<FUICommandList> LoadToolPaletteCommandList;
diff --git a/Engine/Source/Developer/WidgetRegistration/Public/ToolkitBuilder.h b/Engine/Source/Developer/WidgetRegistration/Public/ToolkitBuilder.h
--- a/Engine/Source/Developer/WidgetRegistration/Public/ToolkitBuilder.h
+++ b/Engine/Source/Developer/WidgetRegistration/Public/ToolkitBuilder.h
@@ -205,6 +205,13 @@ public:
 	 */
 	UE_API void InitializeCategoryToolbar(bool bInitLoadToolPaletteMap);
 
+	/**
+	 * Sets the visibility of the category toolbar. Call RefreshCategoryToolbarWidget to apply
+	 * changes in visibility.
+	 * @param Visibility the visibility state to apply
+	 */
+	UE_API void SetCategoryToolbarVisibility(EVisibility Visibility);
+
 	/**
 	 * Sets the display name for the active tool
 	 *
