- Toolkit Builder: Moved 3 members/functions from private to protected so PCG Toolkit Builder has better access
-- This is needed because the preset tool area takes up 50% of the space for the active tool name with no way around it and we need more space.
-- At the same time, we want to use most of the default UI for consistency

diff --git a/Engine/Source/Developer/WidgetRegistration/Public/ToolkitBuilder.h b/Engine/Source/Developer/WidgetRegistration/Public/ToolkitBuilder.h
--- a/Engine/Source/Developer/WidgetRegistration/Public/ToolkitBuilder.h
+++ b/Engine/Source/Developer/WidgetRegistration/Public/ToolkitBuilder.h
@@ -300,9 +300,6 @@ private:
 	 */
 	UE_API TSharedRef<SWidget> GetContextMenuContent(const FName CommandName);
 	
-	/** @return a TSharedPointer to the FToolbarBuilder with the FUICommandInfos that load the Palettes */
-	UE_API TSharedRef<SWidget> GetToolPaletteWidget() const;
-	
 	UE_API void CreatePaletteWidget(FToolPalette& Palette, FToolElement& Element);
 	
 	/**
@@ -334,19 +331,23 @@ private:
 	/** If SelectedCategoryTitleVisibility == EVisibility::Visible, the selected
     * category title is visible, else it is not displayed. By default the selected category title is Collapsed  */	
 	EVisibility SelectedCategoryTitleVisibility = EVisibility::Collapsed;
-
-	/** The FToolkitSections which holds the sections defined for this Toolkit */
-	TSharedPtr<FToolkitSections> ToolkitSections;
-
-	/** If ActivePaletteButtonVisibility == EVisibility::Visible, the command buttons in the active palette
-	 * are visible, otherwise they are not displayed. By default the state is Visible  */
-	EVisibility ActivePaletteButtonVisibility = EVisibility::Visible;
 	
 	/** The display name of the currently active tool */
 	FText ActiveToolDisplayName = FText::GetEmpty();
 
 	/** The array of Command names, in the correct display order */
 	TArray<FName> LoadCommandArray;
+
+protected:
+	/** @return a TSharedPointer to the FToolbarBuilder with the FUICommandInfos that load the Palettes */
+	UE_API TSharedRef<SWidget> GetToolPaletteWidget() const;
+	
+	/** If ActivePaletteButtonVisibility == EVisibility::Visible, the command buttons in the active palette
+	 * are visible, otherwise they are not displayed. By default the state is Visible  */
+	EVisibility ActivePaletteButtonVisibility = EVisibility::Visible;
+
+	/** The FToolkitSections which holds the sections defined for this Toolkit */
+	TSharedPtr<FToolkitSections> ToolkitSections;
 };
 
 #undef UE_API
