Fix crash in slate due to dangling pointers from FSlateBrushTemplates (it happened when adding elements to the array of slate brushes, and reallocating it).

#tests Was crashing 100% when opening the editor before the fix, no more crash. Placement tool is displayed correctly

diff --git a/Engine/Source/Developer/WidgetRegistration/Private/Layout/Containers/ColumnWrappingContainer.cpp b/Engine/Source/Developer/WidgetRegistration/Private/Layout/Containers/ColumnWrappingContainer.cpp
--- a/Engine/Source/Developer/WidgetRegistration/Private/Layout/Containers/ColumnWrappingContainer.cpp
+++ b/Engine/Source/Developer/WidgetRegistration/Private/Layout/Containers/ColumnWrappingContainer.cpp
@@ -56,7 +56,7 @@ void FColumnWrappingContainer::Initialize()
 {
 	SAssignNew (MainContentSBorder,  SBorder )
  		.Padding(8.f)
-		.BorderImage( FSlateBrushTemplates::Get().GetBrushWithColor( EStyleColor::Panel ) )
+		.BorderImage(FSlateBrushTemplates::Panel())
 		[
 			SAssignNew(UniformWrapPanel, SUniformWrapPanel)
 		 .HAlign(HAlign_Fill)
diff --git a/Engine/Source/Developer/WidgetRegistration/Private/Layout/Containers/HeaderAndBodyContainer.cpp b/Engine/Source/Developer/WidgetRegistration/Private/Layout/Containers/HeaderAndBodyContainer.cpp
--- a/Engine/Source/Developer/WidgetRegistration/Private/Layout/Containers/HeaderAndBodyContainer.cpp
+++ b/Engine/Source/Developer/WidgetRegistration/Private/Layout/Containers/HeaderAndBodyContainer.cpp
@@ -60,10 +60,8 @@ TSharedPtr<SWidget> FHeaderAndBodyContainer::GenerateWidget()
 
 	FCurveSequence RolloutCurve = FCurveSequence(0.0f, 1.0f, ECurveEaseFunction::CubicOut);
 
-	constexpr EStyleColor HeaderBorderColor = EStyleColor::Dropdown;
 	constexpr EStyleColor HeaderForegroundColor = EStyleColor::Foreground;
-	constexpr EStyleColor ContainerBackground = EStyleColor::Recessed;
-	const FSlateBrush* ContainerBackgroundBrush = FSlateBrushTemplates::Get().GetBrushWithColor( ContainerBackground );
+	const FSlateBrush* ContainerBackgroundBrush = FSlateBrushTemplates::Recessed();
 
 	// TODO: create style class to put this information in
 	const FName NoBorderButtonStyle = "NoBorder";
@@ -88,7 +86,7 @@ TSharedPtr<SWidget> FHeaderAndBodyContainer::GenerateWidget()
 			.HeightOverride( FOptionalSize() )
 			[
 				SAssignNew( HeaderContentSBorder, SBorder )
-				.BorderImage( FSlateBrushTemplates::Get().GetBrushWithColor( HeaderBorderColor ) )
+				.BorderImage( FSlateBrushTemplates::Dropdown() )
 				.Padding( NoPadding )
 				[
 					SAssignNew( ExpansionButton, SButton )
diff --git a/Engine/Source/Developer/WidgetRegistration/Private/Styles/SlateBrushTemplates.cpp b/Engine/Source/Developer/WidgetRegistration/Private/Styles/SlateBrushTemplates.cpp
--- a/Engine/Source/Developer/WidgetRegistration/Private/Styles/SlateBrushTemplates.cpp
+++ b/Engine/Source/Developer/WidgetRegistration/Private/Styles/SlateBrushTemplates.cpp
@@ -16,6 +16,12 @@ const FSlateBrush* FSlateBrushTemplates::ThinLineHorizontal()
 	return FAppStyle::GetBrush( "ThinLine.Horizontal");		
 }
 
+const FSlateBrush* FSlateBrushTemplates::Dropdown()
+{
+	static const FSlateColorBrush ColorBrush{ FStyleColors::Dropdown };
+	return &ColorBrush;
+}
+
 const FSlateBrush* FSlateBrushTemplates::Transparent()
 {
 	static const FSlateColorBrush ColorBrush{ FLinearColor::Transparent };
@@ -33,7 +39,6 @@ const FSlateBrush* FSlateBrushTemplates::Recessed()
 	static const FSlateColorBrush ColorBrush{ FStyleColors::Recessed };
 	return &ColorBrush;
 }
-
 const FSlateBrush* FSlateBrushTemplates::GetBrushWithColor(EStyleColor Color)
 {
 	const FSlateBrush* ColorBrush = EStyleColorToSlateBrushMap.Find( Color );
diff --git a/Engine/Source/Developer/WidgetRegistration/Public/Styles/SlateBrushTemplates.h b/Engine/Source/Developer/WidgetRegistration/Public/Styles/SlateBrushTemplates.h
--- a/Engine/Source/Developer/WidgetRegistration/Public/Styles/SlateBrushTemplates.h
+++ b/Engine/Source/Developer/WidgetRegistration/Public/Styles/SlateBrushTemplates.h
@@ -15,6 +15,7 @@ struct FSlateBrushTemplates
 	static UE_API const FSlateBrush* ThinLineHorizontal();
 
 	// Colors
+	static UE_API const FSlateBrush* Dropdown();
 	static UE_API const FSlateBrush* Transparent();
 	static UE_API const FSlateBrush* Panel();
 	static UE_API const FSlateBrush* Recessed();
