[CrashFix + Deprecation] Fix for: FSlateBrushTemplates return pointer that can get freed anytime

#Synced-CL 48346153

diff --git a/Engine/Plugins/VirtualProduction/RemoteControl/Source/RemoteControlUI/Private/UI/Behaviour/Builtin/Path/SRCBehaviorSetAssetByPathNewElementRow.cpp b/Engine/Plugins/VirtualProduction/RemoteControl/Source/RemoteControlUI/Private/UI/Behaviour/Builtin/Path/SRCBehaviorSetAssetByPathNewElementRow.cpp
--- a/Engine/Plugins/VirtualProduction/RemoteControl/Source/RemoteControlUI/Private/UI/Behaviour/Builtin/Path/SRCBehaviorSetAssetByPathNewElementRow.cpp
+++ b/Engine/Plugins/VirtualProduction/RemoteControl/Source/RemoteControlUI/Private/UI/Behaviour/Builtin/Path/SRCBehaviorSetAssetByPathNewElementRow.cpp
@@ -70,7 +70,7 @@ void SRCBehaviorSetAssetByPathNewElementRow::Construct(const FArguments& InArgs,
 			.Padding(4.f, 0.f)
 			[
 				SNew(SImage)
-				.Image(FSlateBrushTemplates::DragHandle())
+				.Image(FAppStyle::GetBrush("VerticalBoxDragIndicatorShort"))
 			]
 
 			+ SHorizontalBox::Slot()
diff --git a/Engine/Plugins/VirtualProduction/RemoteControl/Source/RemoteControlUI/Private/UI/Behaviour/SRCBehaviorPanelRow.cpp b/Engine/Plugins/VirtualProduction/RemoteControl/Source/RemoteControlUI/Private/UI/Behaviour/SRCBehaviorPanelRow.cpp
--- a/Engine/Plugins/VirtualProduction/RemoteControl/Source/RemoteControlUI/Private/UI/Behaviour/SRCBehaviorPanelRow.cpp
+++ b/Engine/Plugins/VirtualProduction/RemoteControl/Source/RemoteControlUI/Private/UI/Behaviour/SRCBehaviorPanelRow.cpp
@@ -63,7 +63,7 @@ void SRCBehaviorPanelRow::Construct(const FArguments& InArgs, const TSharedRef<S
 			.Padding(4.f, 0.f)
 			[
 				SNew(SImage)
-				.Image(FSlateBrushTemplates::DragHandle())
+				.Image(FAppStyle::GetBrush("VerticalBoxDragIndicatorShort"))
 			]
 
 			+ SHorizontalBox::Slot()
diff --git a/Engine/Plugins/VirtualProduction/RemoteControl/Source/RemoteControlUI/Private/UI/Controller/CustomControllers/SCustomStringListControllerListEditorWidget.cpp b/Engine/Plugins/VirtualProduction/RemoteControl/Source/RemoteControlUI/Private/UI/Controller/CustomControllers/SCustomStringListControllerListEditorWidget.cpp
--- a/Engine/Plugins/VirtualProduction/RemoteControl/Source/RemoteControlUI/Private/UI/Controller/CustomControllers/SCustomStringListControllerListEditorWidget.cpp
+++ b/Engine/Plugins/VirtualProduction/RemoteControl/Source/RemoteControlUI/Private/UI/Controller/CustomControllers/SCustomStringListControllerListEditorWidget.cpp
@@ -78,7 +78,7 @@ class SCustomStringListControllerEnumEditorRowEditBox : public SCompoundWidget
 				.VAlign(VAlign_Center)
 				[
 					SNew(SImage)
-					.Image(FSlateBrushTemplates::DragHandle())
+					.Image(FAppStyle::GetBrush("VerticalBoxDragIndicatorShort"))
 					.DesiredSizeOverride(FVector2D(5.f, 15.f))
 				];
 		}
diff --git a/Engine/Source/Developer/WidgetRegistration/Private/Layout/Containers/ColumnWrappingContainer.cpp b/Engine/Source/Developer/WidgetRegistration/Private/Layout/Containers/ColumnWrappingContainer.cpp
--- a/Engine/Source/Developer/WidgetRegistration/Private/Layout/Containers/ColumnWrappingContainer.cpp
+++ b/Engine/Source/Developer/WidgetRegistration/Private/Layout/Containers/ColumnWrappingContainer.cpp
@@ -56,7 +56,7 @@ void FColumnWrappingContainer::Initialize()
 {
 	SAssignNew (MainContentSBorder,  SBorder )
  		.Padding(8.f)
-		.BorderImage(FSlateBrushTemplates::Panel())
+		.BorderImage(FAppStyle::GetBrush("Brushes.Panel"))
 		[
 			SAssignNew(UniformWrapPanel, SUniformWrapPanel)
 		 .HAlign(HAlign_Fill)
diff --git a/Engine/Source/Developer/WidgetRegistration/Private/Layout/Containers/HeaderAndBodyContainer.cpp b/Engine/Source/Developer/WidgetRegistration/Private/Layout/Containers/HeaderAndBodyContainer.cpp
--- a/Engine/Source/Developer/WidgetRegistration/Private/Layout/Containers/HeaderAndBodyContainer.cpp
+++ b/Engine/Source/Developer/WidgetRegistration/Private/Layout/Containers/HeaderAndBodyContainer.cpp
@@ -61,7 +61,7 @@ TSharedPtr<SWidget> FHeaderAndBodyContainer::GenerateWidget()
 	FCurveSequence RolloutCurve = FCurveSequence(0.0f, 1.0f, ECurveEaseFunction::CubicOut);
 
 	constexpr EStyleColor HeaderForegroundColor = EStyleColor::Foreground;
-	const FSlateBrush* ContainerBackgroundBrush = FSlateBrushTemplates::Recessed();
+	const FSlateBrush* ContainerBackgroundBrush = FAppStyle::GetBrush("Brushes.Recessed");
 
 	// TODO: create style class to put this information in
 	const FName NoBorderButtonStyle = "NoBorder";
@@ -86,7 +86,7 @@ TSharedPtr<SWidget> FHeaderAndBodyContainer::GenerateWidget()
 			.HeightOverride( FOptionalSize() )
 			[
 				SAssignNew( HeaderContentSBorder, SBorder )
-				.BorderImage( FSlateBrushTemplates::Dropdown() )
+				.BorderImage( FAppStyle::GetBrush("Brushes.Dropdown") )
 				.Padding( NoPadding )
 				[
 					SAssignNew( ExpansionButton, SButton )
diff --git a/Engine/Source/Developer/WidgetRegistration/Private/Styles/SlateBrushTemplates.cpp b/Engine/Source/Developer/WidgetRegistration/Private/Styles/SlateBrushTemplates.cpp
--- a/Engine/Source/Developer/WidgetRegistration/Private/Styles/SlateBrushTemplates.cpp
+++ b/Engine/Source/Developer/WidgetRegistration/Private/Styles/SlateBrushTemplates.cpp
@@ -5,6 +5,9 @@
 #include "Styling/CoreStyle.h"
 #include "Styling/StyleColors.h"
 
+
+PRAGMA_DISABLE_DEPRECATION_WARNINGS
+
 const FSlateBrush* FSlateBrushTemplates::DragHandle()
 {
 	const FSlateBrush* Brush = FCoreStyle::Get().GetBrush("VerticalBoxDragIndicatorShort");
@@ -39,19 +42,12 @@ const FSlateBrush* FSlateBrushTemplates::Recessed()
 	static const FSlateColorBrush ColorBrush{ FStyleColors::Recessed };
 	return &ColorBrush;
 }
+
 const FSlateBrush* FSlateBrushTemplates::GetBrushWithColor(EStyleColor Color)
 {
-	const FSlateBrush* ColorBrush = EStyleColorToSlateBrushMap.Find( Color );
-	
-	if ( ColorBrush )
-	{
-		return ColorBrush;
-	}
+	const TUniquePtr<FSlateColorBrush>& ColorBrushPtr = EStyleColorToSlateBrushMap.FindOrAdd(Color, MakeUnique<FSlateColorBrush>(Color));
 
-	const FSlateColorBrush NewColorBrush{ Color };
-	EStyleColorToSlateBrushMap.Add( Color, static_cast<const FSlateBrush>( NewColorBrush ) );
-
-	return EStyleColorToSlateBrushMap.Find( Color );;
+	return ColorBrushPtr.Get();
 }
 
 FSlateBrushTemplates& FSlateBrushTemplates::Get()
@@ -59,3 +55,5 @@ FSlateBrushTemplates& FSlateBrushTemplates::Get()
 	static FSlateBrushTemplates Templates;
 	return Templates;
 }
+
+PRAGMA_ENABLE_DEPRECATION_WARNINGS
diff --git a/Engine/Source/Developer/WidgetRegistration/Public/Layout/SeparatorBuilder.h b/Engine/Source/Developer/WidgetRegistration/Public/Layout/SeparatorBuilder.h
--- a/Engine/Source/Developer/WidgetRegistration/Public/Layout/SeparatorBuilder.h
+++ b/Engine/Source/Developer/WidgetRegistration/Public/Layout/SeparatorBuilder.h
@@ -135,12 +135,12 @@ private:
 	FSeparatorBuilder(EStyleColor InColor = EStyleColor::Background,
 		              EOrientation InOrientation = Orient_Horizontal,
 					  FSeparatorSize InSize = FSeparatorSizes::Small,
-					  const FSlateBrush* InImage = FSlateBrushTemplates::ThinLineHorizontal());
+					  const FSlateBrush* InImage = FAppStyle::GetBrush("ThinLine.Horizontal"));
 
 	/**
 	 * The const FSlateBrush* for the SSeparator
 	 */
-	const FSlateBrush* Image =  FSlateBrushTemplates::ThinLineHorizontal();
+	const FSlateBrush* Image = FAppStyle::GetBrush("ThinLine.Horizontal");
 
 	/**
 	 * The EOrientation that specifies whether this is a horizontal or vertical Separator
diff --git a/Engine/Source/Developer/WidgetRegistration/Public/Styles/SlateBrushTemplates.h b/Engine/Source/Developer/WidgetRegistration/Public/Styles/SlateBrushTemplates.h
--- a/Engine/Source/Developer/WidgetRegistration/Public/Styles/SlateBrushTemplates.h
+++ b/Engine/Source/Developer/WidgetRegistration/Public/Styles/SlateBrushTemplates.h
@@ -1,6 +1,7 @@
 // Copyright Epic Games, Inc. All Rights Reserved.
 
 #pragma once
+#include "Brushes/SlateColorBrush.h"
 #include "Styling/AppStyle.h"
 #include "Styling/SlateBrush.h"
 #include "Styling/StyleColors.h"
@@ -8,16 +9,22 @@
 #define UE_API WIDGETREGISTRATION_API
 
 /** A const FSlateBrush* Factory */
-struct FSlateBrushTemplates
+struct UE_DEPRECATED(5.8, "This struct is deprecated. You should use FAppStyle instead") FSlateBrushTemplates
 {
 	// Images
+	// UE_DEPRECATED 5.8 : Use FAppStyle::GetBrush("VerticalBoxDragIndicatorShort") instead
 	static UE_API const FSlateBrush* DragHandle();
+	// UE_DEPRECATED 5.8 : "Use FAppStyle::GetBrush("ThinLine.Horizontal") instead
 	static UE_API const FSlateBrush* ThinLineHorizontal();
 
 	// Colors
+	// UE_DEPRECATED 5.8 : Use FAppStyle::GetBrush("Brushes.Dropdown") instead
 	static UE_API const FSlateBrush* Dropdown();
+	// UE_DEPRECATED 5.8 : Use FAppStyle::GetBrush("Brushes.Transparent") instead
 	static UE_API const FSlateBrush* Transparent();
+	// UE_DEPRECATED 5.8 : Use FAppStyle::GetBrush("Brushes.Panel") instead
 	static UE_API const FSlateBrush* Panel();
+	// UE_DEPRECATED 5.8 : Use FAppStyle::GetBrush("Brushes.Recessed") instead
 	static UE_API const FSlateBrush* Recessed();
 
 	/**
@@ -25,15 +32,17 @@ struct FSlateBrushTemplates
 	 *
 	 *  @param the EStyleColor we need the slate brush for
 	 */
+	// UE_DEPRECATED 5.8 : Use FAppStyle::GetBrush() with a predefined brush name in StyleSheets instead
 	UE_API const FSlateBrush* GetBrushWithColor(EStyleColor Color);
 
 	/**
 	 * gets the FSlateBrushTemplates singleton
 	 */
+	//UE_DEPRECATED 5.8: This struct is deprecated. You should use FAppStyle instead
 	static UE_API FSlateBrushTemplates& Get();
 
-	/** the map of EStyleColor to const FSlateBrush */
-	TMap<EStyleColor, const FSlateBrush> EStyleColorToSlateBrushMap;
+	/** the map of EStyleColor to const FSlateColorBrush UniquePtr */
+	TMap<EStyleColor, TUniquePtr<FSlateColorBrush>> EStyleColorToSlateBrushMap;
 };
 
 #undef UE_API
diff --git a/Engine/Source/Editor/PlacementMode/Private/SPlacementModeTools.cpp b/Engine/Source/Editor/PlacementMode/Private/SPlacementModeTools.cpp
--- a/Engine/Source/Editor/PlacementMode/Private/SPlacementModeTools.cpp
+++ b/Engine/Source/Editor/PlacementMode/Private/SPlacementModeTools.cpp
@@ -230,7 +230,7 @@ void SPlacementAssetEntry::Construct(const FArguments& InArgs, const TSharedPtr<
 	AssetTextPadding = FMargin{ 4, 0, 8, 1 };
 	DraggableAssetEndWidget = SNew(SBox )
 								.Padding( DragHandlePadding )
-								[  SNew(SImage).Image( FSlateBrushTemplates::DragHandle() ) ];
+								[  SNew(SImage).Image( FAppStyle::GetBrush("VerticalBoxDragIndicatorShort") ) ];
 	ThumbnailBoxWidth = 20;
 	
 	const FSlateBrush* CustomIconBrush = nullptr;
@@ -710,7 +710,7 @@ void SPlacementModeTools::Construct( const FArguments& InArgs, TSharedRef<SDockT
 			.FillHeight(1)
 			[
 				SNew(SBorder)
-				.BorderImage( FSlateBrushTemplates::Panel() )
+				.BorderImage( FAppStyle::GetBrush("Brushes.Panel") )
 				.Padding(0)
 				[
 					CategoryContentBuilder->GenerateWidgetSharedRef()
diff --git a/Engine/Source/Runtime/SlateCore/Private/Styling/StarshipCoreStyle.cpp b/Engine/Source/Runtime/SlateCore/Private/Styling/StarshipCoreStyle.cpp
--- a/Engine/Source/Runtime/SlateCore/Private/Styling/StarshipCoreStyle.cpp
+++ b/Engine/Source/Runtime/SlateCore/Private/Styling/StarshipCoreStyle.cpp
@@ -1100,6 +1100,7 @@ void FStarshipCoreStyle::SetStyle(const TSharedRef< ISlateStyle >& NewStyle)
 
 void FStarshipCoreStyle::SetupColors(TSharedRef<FStyle>& Style)
 {
+	Style->Set("Colors.Transparent", FStyleColors::Transparent);
 	Style->Set("Colors.Black", FStyleColors::Black);
 	Style->Set("Colors.Title", FStyleColors::Title);
 	Style->Set("Colors.WindowBorder", FStyleColors::WindowBorder);
@@ -1149,6 +1150,7 @@ void FStarshipCoreStyle::SetupColors(TSharedRef<FStyle>& Style)
 	Style->Set("Colors.Warning", FStyleColors::Warning);
 	Style->Set("Colors.Error", FStyleColors::Error);
 
+	Style->Set("Brushes.Transparent", new FSlateColorBrush(FStyleColors::Transparent));
 	Style->Set("Brushes.Black", new FSlateColorBrush(FStyleColors::Black));
 	Style->Set("Brushes.Title", new FSlateColorBrush(FStyleColors::Title));
 	Style->Set("Brushes.Foldout", new FSlateColorBrush(FStyleColors::Foldout));
