Fix compilation errors when UE_SLATE_WITH_DYNAMIC_INVALIDATION is disabled

#Synced-CL 50475317

diff --git a/Engine/Source/Developer/WidgetRegistration/Private/Layout/CategoryDrivenContentBuilderBase.cpp b/Engine/Source/Developer/WidgetRegistration/Private/Layout/CategoryDrivenContentBuilderBase.cpp
--- a/Engine/Source/Developer/WidgetRegistration/Private/Layout/CategoryDrivenContentBuilderBase.cpp
+++ b/Engine/Source/Developer/WidgetRegistration/Private/Layout/CategoryDrivenContentBuilderBase.cpp
@@ -145,21 +145,29 @@ void FCategoryDrivenContentBuilderBase::CreateWidget()
 		.Resizable(false)
 		.SizeRule(SSplitter::SizeToContent)
 		[
+#if UE_SLATE_WITH_DYNAMIC_INVALIDATION
 			SNew(SInvalidationPanel)
 			.UseDynamicInvalidation(true)
 			[
+#endif
 				CategoryToolbarVBox.ToSharedRef()
+#if UE_SLATE_WITH_DYNAMIC_INVALIDATION
 			]
+#endif
 		]
 
 		+ SSplitter::Slot()
 		.SizeRule(SSplitter::FractionOfParent)
 		[
+#if UE_SLATE_WITH_DYNAMIC_INVALIDATION
 			SNew(SInvalidationPanel)
 			.UseDynamicInvalidation(true)
 			[
+#endif
 				MainContentVerticalBox->AsShared()
+#if UE_SLATE_WITH_DYNAMIC_INVALIDATION
 			]
+#endif
 		];
 	
 		ToolkitWidgetContainerVBox->AddSlot()
diff --git a/Engine/Source/Editor/ContentBrowser/Private/SContentBrowser.cpp b/Engine/Source/Editor/ContentBrowser/Private/SContentBrowser.cpp
--- a/Engine/Source/Editor/ContentBrowser/Private/SContentBrowser.cpp
+++ b/Engine/Source/Editor/ContentBrowser/Private/SContentBrowser.cpp
@@ -1065,9 +1065,11 @@ TSharedRef<SWidget> SContentBrowser::CreateAssetView(const FContentBrowserConfig
 			.BorderImage(FAppStyle::Get().GetBrush("Brushes.Panel"))
 			.Padding(AssetViewPadding)
 			[
+#if UE_SLATE_WITH_DYNAMIC_INVALIDATION
 				SNew(SInvalidationPanel)
 				.UseDynamicInvalidation(true)
 				[
+#endif
 					SNew(SVerticalBox)
 
 					+ SVerticalBox::Slot()
@@ -1076,7 +1078,9 @@ TSharedRef<SWidget> SContentBrowser::CreateAssetView(const FContentBrowserConfig
 					[
 						LegacyContentSourceWidgets->AssetViewPtr.ToSharedRef()
 					]
+#if UE_SLATE_WITH_DYNAMIC_INVALIDATION
 				]
+#endif
 			];
 
 	return LegacyContentSourceWidgets->AssetViewBorder.ToSharedRef();
diff --git a/Engine/Source/Editor/ContentBrowser/Private/Widgets/SContentBrowserSourceTree.cpp b/Engine/Source/Editor/ContentBrowser/Private/Widgets/SContentBrowserSourceTree.cpp
--- a/Engine/Source/Editor/ContentBrowser/Private/Widgets/SContentBrowserSourceTree.cpp
+++ b/Engine/Source/Editor/ContentBrowser/Private/Widgets/SContentBrowserSourceTree.cpp
@@ -438,16 +438,20 @@ namespace UE::Editor::ContentBrowser::Private
 					+ SVerticalBox::Slot()
 					.Padding(FMargin(0, 1))
 					[
+#if UE_SLATE_WITH_DYNAMIC_INVALIDATION
 						SNew(SInvalidationPanel)
 						.UseDynamicInvalidation(true)
 						[
+#endif
 							// Surround scrollable with a scrollbox (adds drop shadows)
 							SNew(SScrollBorder, InBody)
 							[
 								BodyScrollableWidget->GetScrollWidget()
 							]
 							.Style(InArgs._bEnableShadowBoxStyle ? InArgs._ShadowBoxStyle : &FCoreStyle::Get().GetWidgetStyle<FScrollBorderStyle>("ScrollBorder"))
+#if UE_SLATE_WITH_DYNAMIC_INVALIDATION
 						]
+#endif
 					]
 				]
 
diff --git a/Engine/Source/Editor/GraphEditor/Private/SGraphEditorImpl.cpp b/Engine/Source/Editor/GraphEditor/Private/SGraphEditorImpl.cpp
--- a/Engine/Source/Editor/GraphEditor/Private/SGraphEditorImpl.cpp
+++ b/Engine/Source/Editor/GraphEditor/Private/SGraphEditorImpl.cpp
@@ -754,18 +754,23 @@ void SGraphEditorImpl::Construct( const FArguments& InArgs )
 		.OnNodeSingleClicked(InArgs._GraphEvents.OnNodeSingleClicked);
 
 	TSharedPtr<SWidget> GraphPanelWidget;
+#if UE_SLATE_WITH_DYNAMIC_INVALIDATION
 	if (InArgs._GraphSupportsInvalidation)
 	{
 		GraphPanelWidget = SNew(SInvalidationPanel)
 		.UseDynamicInvalidation(true)
+
 		[
 			GraphPanel.ToSharedRef()
 		];
 	}
 	else
 	{
+#endif
 		GraphPanelWidget = GraphPanel;
+#if UE_SLATE_WITH_DYNAMIC_INVALIDATION
 	}
+#endif
 
 	this->ChildSlot
 	[
diff --git a/Engine/Source/Editor/SceneOutliner/Private/SSceneOutliner.cpp b/Engine/Source/Editor/SceneOutliner/Private/SSceneOutliner.cpp
--- a/Engine/Source/Editor/SceneOutliner/Private/SSceneOutliner.cpp
+++ b/Engine/Source/Editor/SceneOutliner/Private/SSceneOutliner.cpp
@@ -263,9 +263,11 @@ void SSceneOutliner::Construct(const FArguments& InArgs, const FSceneOutlinerIni
 
 		+SOverlay::Slot()
 		[
+#if UE_SLATE_WITH_DYNAMIC_INVALIDATION
 			SNew(SInvalidationPanel)
 			.UseDynamicInvalidation(true)
 			[
+#endif
 				SAssignNew( OutlinerTreeView, SSceneOutlinerTreeView, StaticCastSharedRef<SSceneOutliner>(AsShared()) )
 
 				// Determined by the mode
@@ -315,7 +317,9 @@ void SSceneOutliner::Construct(const FArguments& InArgs, const FSceneOutlinerIni
 
 				// Preserve the selection when the selected item is hidden due to a parent collapsing
 				.AllowInvisibleItemSelection(true)
+#if UE_SLATE_WITH_DYNAMIC_INVALIDATION
 			]
+#endif
 		]
 	];
 
